Fixes compile errors with latest version of Visual C++ compiler in /permissive- mode.

diff --git a/Engine/Plugins/Media/AjaMedia/Source/AjaMediaEditor/Private/Customizations/AjaMediaModeCustomization.cpp b/Engine/Plugins/Media/AjaMedia/Source/AjaMediaEditor/Private/Customizations/AjaMediaModeCustomization.cpp
--- a/Engine/Plugins/Media/AjaMedia/Source/AjaMediaEditor/Private/Customizations/AjaMediaModeCustomization.cpp
+++ b/Engine/Plugins/Media/AjaMedia/Source/AjaMediaEditor/Private/Customizations/AjaMediaModeCustomization.cpp
@@ -83,7 +83,7 @@ TSharedRef<SWidget> FAjaMediaModeCustomization::HandleSourceComboButtonMenuConte
 	// generate menu
 	FMenuBuilder MenuBuilder(true, nullptr);
 
-	ANSICHAR* SectionName = bOutput ? "AllOutputModes" : "AllInputModes";
+	const ANSICHAR* SectionName = bOutput ? "AllOutputModes" : "AllInputModes";
 	TAttribute<FText> HeaderText = bOutput ? LOCTEXT("AllOutputModesSection", "Output Modes") : LOCTEXT("AllInputModesSection", "Input Modes");
 
 	MenuBuilder.BeginSection(SectionName, HeaderText);
