Remove EngineVersion for AJA uplugin

diff --git a/Engine/Plugins/Media/AjaMedia/AjaMedia.uplugin b/Engine/Plugins/Media/AjaMedia/AjaMedia.uplugin
--- a/Engine/Plugins/Media/AjaMedia/AjaMedia.uplugin
+++ b/Engine/Plugins/Media/AjaMedia/AjaMedia.uplugin
@@ -6,7 +6,6 @@
 	"VersionName" : "1.0",
 	"CreatedBy" : "Epic Games Inc",
 	"CreatedByURL" : "http://epicgames.com",
-	"EngineVersion" : "4.20.0",
 	"Description" : "Implements input and output using Aja Capture cards.",
 	"Category" : "Media Players",
 	"CanContainContent" : false,
