AJA late feature fix build.

diff --git a/Engine/Plugins/Media/AjaMedia/AjaMedia.uplugin b/Engine/Plugins/Media/AjaMedia/AjaMedia.uplugin
--- a/Engine/Plugins/Media/AjaMedia/AjaMedia.uplugin
+++ b/Engine/Plugins/Media/AjaMedia/AjaMedia.uplugin
@@ -28,7 +28,8 @@
 		{
 			"Name": "AjaMediaFactory",
 			"Type": "Editor",
-			"LoadingPhase": "PostEngineInit"
+			"LoadingPhase": "PostEngineInit",
+			"WhitelistPlatforms": [ "Win64" ]
 		},
 		{
 			"Name": "AjaMediaFactory",
@@ -39,7 +40,8 @@
 		{
 			"Name": "AjaMediaEditor",
 			"Type": "Editor",
-			"LoadingPhase": "PostEngineInit"
+			"LoadingPhase": "PostEngineInit",
+			"WhitelistPlatforms": [ "Win64" ]
 		}
 	],
 	"Plugins" :
diff --git a/Engine/Plugins/Media/AjaMedia/Source/AjaMediaOutput/Private/AjaMediaOutput.cpp b/Engine/Plugins/Media/AjaMedia/Source/AjaMediaOutput/Private/AjaMediaOutput.cpp
--- a/Engine/Plugins/Media/AjaMedia/Source/AjaMediaOutput/Private/AjaMediaOutput.cpp
+++ b/Engine/Plugins/Media/AjaMedia/Source/AjaMediaOutput/Private/AjaMediaOutput.cpp
@@ -95,7 +95,7 @@ FAjaMediaMode UAjaMediaOutput::GetMediaMode() const
 
 FIntPoint UAjaMediaOutput::GetRequestedSize() const
 {
-	return MediaMode.TargetSize;
+	return GetMediaMode().TargetSize;
 }
 
 EPixelFormat UAjaMediaOutput::GetRequestedPixelFormat() const
