AJA Late feature

- Adding missing default plugin MediaBundle assets for 4 ports

diff --git a/Engine/Plugins/Media/AjaMedia/AjaMedia.uplugin b/Engine/Plugins/Media/AjaMedia/AjaMedia.uplugin
--- a/Engine/Plugins/Media/AjaMedia/AjaMedia.uplugin
+++ b/Engine/Plugins/Media/AjaMedia/AjaMedia.uplugin
@@ -8,7 +8,7 @@
 	"CreatedByURL" : "http://epicgames.com",
 	"Description" : "Implements input and output using Aja Capture cards.",
 	"Category" : "Media Players",
-	"CanContainContent" : false,
+	"CanContainContent" : true,
 	"EnabledByDefault" : false,
 	"IsBetaVersion": true,
 
