Aja

- Adding a check to CanEverRender flag to mark the Aja card usable or not. Will still give warnings to notify user that Aja card isn't usable.

diff --git a/Engine/Plugins/Media/AjaMedia/Source/AjaMedia/Private/Aja/Aja.cpp b/Engine/Plugins/Media/AjaMedia/Source/AjaMedia/Private/Aja/Aja.cpp
--- a/Engine/Plugins/Media/AjaMedia/Source/AjaMedia/Private/Aja/Aja.cpp
+++ b/Engine/Plugins/Media/AjaMedia/Source/AjaMedia/Private/Aja/Aja.cpp
@@ -11,6 +11,7 @@
  //~ Static initialization
  //--------------------------------------------------------------------
 void* FAja::LibHandle = nullptr; 
+bool FAja::bCanForceAJAUsage = false; 
 
 //~ Initialization functions implementation
 //--------------------------------------------------------------------
@@ -44,6 +45,9 @@ bool FAja::Initialize()
 		return false;
 	}
 
+	//Check if command line argument to force AJA card usage is there
+	bCanForceAJAUsage = FParse::Param(FCommandLine::Get(), TEXT("forceajausage"));
+
 #if !NO_LOGGING
 	AJA::SetLoggingCallbacks(&LogInfo, &LogWarning, &LogError);
 #endif // !NO_LOGGING
diff --git a/Engine/Plugins/Media/AjaMedia/Source/AjaMedia/Private/Aja/Aja.h b/Engine/Plugins/Media/AjaMedia/Source/AjaMedia/Private/Aja/Aja.h
--- a/Engine/Plugins/Media/AjaMedia/Source/AjaMedia/Private/Aja/Aja.h
+++ b/Engine/Plugins/Media/AjaMedia/Source/AjaMedia/Private/Aja/Aja.h
@@ -3,6 +3,7 @@
 #pragma once
 
 #include "CoreTypes.h"
+#include "Misc/App.h"
 #include "Misc/Timecode.h"
 #include "Misc/Timespan.h"
 
@@ -23,6 +24,8 @@ public:
 	static FTimespan ConvertAJATimecode2Timespan(const AJA::FTimecode& InTimecode, const AJA::FTimecode& PreviousTimeCode, const FTimespan& PreviousTimespan, const FFrameRate& InFPS);
 	static FTimecode ConvertAJATimecode2Timecode(const AJA::FTimecode& InTimecode, const FFrameRate& InFPS);
 
+	static bool CanUseAJACard() { return (FApp::CanEverRender() || bCanForceAJAUsage); }
+
 private:
 	static void LogInfo(const TCHAR* InFormat, ...);
 	static void LogWarning(const TCHAR* InFormat, ...);
@@ -30,4 +33,5 @@ private:
 
 private:
 	static void* LibHandle;
+	static bool bCanForceAJAUsage;
 };
diff --git a/Engine/Plugins/Media/AjaMedia/Source/AjaMedia/Private/AjaMediaModule.cpp b/Engine/Plugins/Media/AjaMedia/Source/AjaMedia/Private/AjaMediaModule.cpp
--- a/Engine/Plugins/Media/AjaMedia/Source/AjaMedia/Private/AjaMediaModule.cpp
+++ b/Engine/Plugins/Media/AjaMedia/Source/AjaMedia/Private/AjaMediaModule.cpp
@@ -51,6 +51,8 @@ public:
 
 	virtual bool IsInitialized() const override { return FAja::IsInitialized(); }
 
+	virtual bool CanBeUsed() const override { return FAja::CanUseAJACard(); }
+
 public:
 
 	//~ IModuleInterface interface
diff --git a/Engine/Plugins/Media/AjaMedia/Source/AjaMedia/Private/Assets/AjaCustomTimeStep.cpp b/Engine/Plugins/Media/AjaMedia/Source/AjaMedia/Private/Assets/AjaCustomTimeStep.cpp
--- a/Engine/Plugins/Media/AjaMedia/Source/AjaMedia/Private/Assets/AjaCustomTimeStep.cpp
+++ b/Engine/Plugins/Media/AjaMedia/Source/AjaMedia/Private/Assets/AjaCustomTimeStep.cpp
@@ -66,6 +66,13 @@ bool UAjaCustomTimeStep::Initialize(UEngine* InEngine)
 		return false;
 	}
 
+	if (!FAja::CanUseAJACard())
+	{
+		State = ECustomTimeStepSynchronizationState::Error;
+		UE_LOG(LogAjaMedia, Warning, TEXT("The CustomTimeStep '%s' can't be initialized because Aja card cannot be used. Are you in a Commandlet? You may override this behavior by launching with -ForceAjaUsage"), *GetName());
+		return false;
+	}
+
 	const FAjaMediaMode CurrentMediaMode = GetMediaMode();
 
 	FString FailureReason;
diff --git a/Engine/Plugins/Media/AjaMedia/Source/AjaMedia/Private/Assets/AjaMediaFinder.cpp b/Engine/Plugins/Media/AjaMedia/Source/AjaMedia/Private/Assets/AjaMediaFinder.cpp
--- a/Engine/Plugins/Media/AjaMedia/Source/AjaMedia/Private/Assets/AjaMediaFinder.cpp
+++ b/Engine/Plugins/Media/AjaMedia/Source/AjaMedia/Private/Assets/AjaMediaFinder.cpp
@@ -157,7 +157,7 @@ bool FAjaMediaMode::IsValid() const
 bool FAjaMediaFinder::GetSources(TArray<FAjaMediaPort>& OutSources)
 {
 	OutSources.Reset();
-	if (!FAja::IsInitialized())
+	if (!FAja::IsInitialized() || !FAja::CanUseAJACard())
 	{
 		return false;
 	}
@@ -191,7 +191,7 @@ bool FAjaMediaFinder::GetSources(TArray<FAjaMediaPort>& OutSources)
 bool FAjaMediaFinder::GetModes(int32 DeviceIndex, bool bInOutput, TArray<FAjaMediaMode>& OutModes)
 {
 	OutModes.Reset();
-	if (!FAja::IsInitialized())
+	if (!FAja::IsInitialized() || !FAja::CanUseAJACard())
 	{
 		return false;
 	}
diff --git a/Engine/Plugins/Media/AjaMedia/Source/AjaMedia/Private/Assets/AjaTimecodeProvider.cpp b/Engine/Plugins/Media/AjaMedia/Source/AjaMedia/Private/Assets/AjaTimecodeProvider.cpp
--- a/Engine/Plugins/Media/AjaMedia/Source/AjaMedia/Private/Assets/AjaTimecodeProvider.cpp
+++ b/Engine/Plugins/Media/AjaMedia/Source/AjaMedia/Private/Assets/AjaTimecodeProvider.cpp
@@ -76,6 +76,13 @@ bool UAjaTimecodeProvider::Initialize(class UEngine* InEngine)
 		return false;
 	}
 
+	if (!FAja::CanUseAJACard())
+	{
+		State = ETimecodeProviderSynchronizationState::Error;
+		UE_LOG(LogAjaMedia, Warning, TEXT("The TimecodeProvider '%s' can't be initialized because Aja card cannot be used. Are you in a Commandlet? You may override this behavior by launching with -ForceAjaUsage"), *GetName());
+		return false;
+	}
+
 	const FAjaMediaMode CurrentMediaMode = GetMediaMode();
 
 	FString FailureReason;
diff --git a/Engine/Plugins/Media/AjaMedia/Source/AjaMedia/Private/Player/AjaMediaPlayer.cpp b/Engine/Plugins/Media/AjaMedia/Source/AjaMedia/Private/Player/AjaMediaPlayer.cpp
--- a/Engine/Plugins/Media/AjaMedia/Source/AjaMedia/Private/Player/AjaMediaPlayer.cpp
+++ b/Engine/Plugins/Media/AjaMedia/Source/AjaMedia/Private/Player/AjaMediaPlayer.cpp
@@ -87,6 +87,12 @@ bool FAjaMediaPlayer::Open(const FString& Url, const IMediaOptions* Options)
 		return false;
 	}
 
+	if (!FAja::CanUseAJACard())
+	{
+		UE_LOG(LogAjaMedia, Warning, TEXT("The AjaMediaPlayer can't open URL '%s' because Aja card cannot be used. Are you in a Commandlet? You may override this behavior by launching with -ForceAjaUsage"), *DeviceSource.ToString());
+		return false;
+	}
+
 	AJA::AJADeviceOptions DeviceOptions(DeviceSource.DeviceIndex);
 
 	// Read options
diff --git a/Engine/Plugins/Media/AjaMedia/Source/AjaMedia/Public/IAjaMediaModule.h b/Engine/Plugins/Media/AjaMedia/Source/AjaMedia/Public/IAjaMediaModule.h
--- a/Engine/Plugins/Media/AjaMedia/Source/AjaMedia/Public/IAjaMediaModule.h
+++ b/Engine/Plugins/Media/AjaMedia/Source/AjaMedia/Public/IAjaMediaModule.h
@@ -24,5 +24,8 @@ public:
 
 	/** @return true if the Aja module and Aja dll could be loaded */
 	virtual bool IsInitialized() const = 0;
+	
+	/** @return true if the Aja card can be used */
+	virtual bool CanBeUsed() const = 0;
 };
 
diff --git a/Engine/Plugins/Media/AjaMedia/Source/AjaMediaOutput/Private/AjaMediaCapture.cpp b/Engine/Plugins/Media/AjaMedia/Source/AjaMediaOutput/Private/AjaMediaCapture.cpp
--- a/Engine/Plugins/Media/AjaMedia/Source/AjaMediaOutput/Private/AjaMediaCapture.cpp
+++ b/Engine/Plugins/Media/AjaMedia/Source/AjaMediaOutput/Private/AjaMediaCapture.cpp
@@ -2,6 +2,7 @@
 
 #include "AjaMediaCapture.h"
 
+#include "IAjaMediaModule.h"
 #include "AjaMediaOutput.h"
 #include "Engine/RendererSettings.h"
 #include "HAL/Event.h"
@@ -163,6 +164,13 @@ bool UAjaMediaCapture::InitAJA(UAjaMediaOutput* InAjaMediaOutput)
 {
 	check(InAjaMediaOutput);
 
+	IAjaMediaModule& MediaModule = FModuleManager::LoadModuleChecked<IAjaMediaModule>(TEXT("AjaMedia"));
+	if (!MediaModule.CanBeUsed())
+	{
+		UE_LOG(LogAjaMediaOutput, Warning, TEXT("The AjaMediaCapture can't open MediaOutput '%s' because Aja card cannot be used. Are you in a Commandlet? You may override this behavior by launching with -ForceAjaUsage"), *InAjaMediaOutput->GetName());
+		return false;
+	}
+
 	// Init general settings
 	bWaitForSyncEvent = InAjaMediaOutput->bWaitForSyncEvent;
 	bEncodeTimecodeInTexel = InAjaMediaOutput->bEncodeTimecodeInTexel;
