- Updating MediaIO to work with TimedData monitor
- Disabled LiveLink' ValidEngineTime and ValidTimecodeFrame setting by default.
- Add option to set the LiveLink source to Timecode by default.

diff --git a/Engine/Plugins/Animation/LiveLink/Source/LiveLink/Private/LiveLinkMessageBusSource.cpp b/Engine/Plugins/Animation/LiveLink/Source/LiveLink/Private/LiveLinkMessageBusSource.cpp
--- a/Engine/Plugins/Animation/LiveLink/Source/LiveLink/Private/LiveLinkMessageBusSource.cpp
+++ b/Engine/Plugins/Animation/LiveLink/Source/LiveLink/Private/LiveLinkMessageBusSource.cpp
@@ -16,6 +16,21 @@
 #include "MessageEndpointBuilder.h"
 #include "Misc/App.h"
 
+
+FLiveLinkMessageBusSource::FLiveLinkMessageBusSource(const FText& InSourceType, const FText& InSourceMachineName, const FMessageAddress& InConnectionAddress, double InMachineTimeOffset)
+	: ConnectionAddress(InConnectionAddress)
+	, SourceType(InSourceType)
+	, SourceMachineName(InSourceMachineName)
+	, ConnectionLastActive(0.0)
+	, bIsValid(false)
+	, MachineTimeOffset(InMachineTimeOffset)
+{}
+
+void FLiveLinkMessageBusSource::InitializeSettings(ULiveLinkSourceSettings* Settings)
+{
+	Settings->Mode = GetDefault<ULiveLinkSettings>()->DefaultMessageBusSourceMode;
+}
+
 void FLiveLinkMessageBusSource::ReceiveClient(ILiveLinkClient* InClient, FGuid InSourceGuid)
 {
 	Client = InClient;
diff --git a/Engine/Plugins/Animation/LiveLink/Source/LiveLink/Private/LiveLinkMessageBusSource.h b/Engine/Plugins/Animation/LiveLink/Source/LiveLink/Private/LiveLinkMessageBusSource.h
--- a/Engine/Plugins/Animation/LiveLink/Source/LiveLink/Private/LiveLinkMessageBusSource.h
+++ b/Engine/Plugins/Animation/LiveLink/Source/LiveLink/Private/LiveLinkMessageBusSource.h
@@ -20,15 +20,10 @@ class LIVELINK_API FLiveLinkMessageBusSource : public ILiveLinkSource
 {
 public:
 
-	FLiveLinkMessageBusSource(const FText& InSourceType, const FText& InSourceMachineName, const FMessageAddress& InConnectionAddress, double InMachineTimeOffset)
-		: ConnectionAddress(InConnectionAddress)
-		, SourceType(InSourceType)
-		, SourceMachineName(InSourceMachineName)
-		, ConnectionLastActive(0.0)
-		, bIsValid(false)
-		, MachineTimeOffset(InMachineTimeOffset)
-	{}
+	FLiveLinkMessageBusSource(const FText& InSourceType, const FText& InSourceMachineName, const FMessageAddress& InConnectionAddress, double InMachineTimeOffset);
 
+	//~ Begin ILiveLinkSource interface
+	virtual void InitializeSettings(ULiveLinkSourceSettings* Settings) override;
 	virtual void ReceiveClient(ILiveLinkClient* InClient, FGuid InSourceGuid) override;
 	virtual void Update() override;
 
@@ -39,6 +34,7 @@ public:
 	virtual FText GetSourceType() const override { return SourceType; }
 	virtual FText GetSourceMachineName() const override { return SourceMachineName; }
 	virtual FText GetSourceStatus() const override;
+	//~ End ILiveLinkSource interface
 
 private:
 	//~ Message bus message handlers
diff --git a/Engine/Plugins/Animation/LiveLink/Source/LiveLink/Private/LiveLinkSettings.cpp b/Engine/Plugins/Animation/LiveLink/Source/LiveLink/Private/LiveLinkSettings.cpp
--- a/Engine/Plugins/Animation/LiveLink/Source/LiveLink/Private/LiveLinkSettings.cpp
+++ b/Engine/Plugins/Animation/LiveLink/Source/LiveLink/Private/LiveLinkSettings.cpp
@@ -11,7 +11,8 @@ FLiveLinkRoleProjectSetting::FLiveLinkRoleProjectSetting()
 
 
 ULiveLinkSettings::ULiveLinkSettings()
-	: MessageBusPingRequestFrequency(1.0)
+	: DefaultMessageBusSourceMode(ELiveLinkSourceMode::EngineTime)
+	, MessageBusPingRequestFrequency(1.0)
 	, MessageBusHeartbeatFrequency(1.0)
 	, MessageBusHeartbeatTimeout(2.0)
 	, MessageBusTimeBeforeRemovingInactiveSource(30.0)
diff --git a/Engine/Plugins/Animation/LiveLink/Source/LiveLink/Private/LiveLinkSubject.cpp b/Engine/Plugins/Animation/LiveLink/Source/LiveLink/Private/LiveLinkSubject.cpp
--- a/Engine/Plugins/Animation/LiveLink/Source/LiveLink/Private/LiveLinkSubject.cpp
+++ b/Engine/Plugins/Animation/LiveLink/Source/LiveLink/Private/LiveLinkSubject.cpp
@@ -71,40 +71,15 @@ void FLiveLinkSubject::Update()
 
 	if (GetMode() == ELiveLinkSourceMode::EngineTime)
 	{
-		double ValidEngineTime = FApp::GetCurrentTime() - CachedSettings.BufferSettings.EngineTimeOffset - CachedSettings.BufferSettings.ValidEngineTime;
-		int32 FrameIndex = 0;
-		for (const FLiveLinkFrameDataStruct& SourceFrameData : FrameData)
+		if (CachedSettings.BufferSettings.bValidEngineTimeEnabled)
 		{
-			double FrameTime = SourceFrameData.GetBaseData()->WorldTime.GetOffsettedTime();
-			double OffsetTime = ValidEngineTime;
-			if (FrameTime > OffsetTime)
-			{
-				break;
-			}
-			++FrameIndex;
-		}
-
-		if (FrameIndex - 1 >= 0)
-		{
-			const int32 Count = CachedSettings.BufferSettings.bKeepAtLeastOneFrame && FrameData.Num() == FrameIndex ? FrameIndex - 1 : FrameIndex;
-			if (Count > 0)
-			{
-				FrameData.RemoveAt(0, Count, false);
-			}
-		}
-	}
-	else if (GetMode() == ELiveLinkSourceMode::Timecode)
-	{
-		if (FApp::GetCurrentFrameTime().IsSet())
-		{
-			const FQualifiedFrameTime CurrentSyncTime = FApp::GetCurrentFrameTime().GetValue();
-			const FFrameTime CurrentFrameTimeInFrameSpace = CurrentSyncTime.ConvertTo(CachedSettings.BufferSettings.TimecodeFrameRate);
+			double ValidEngineTime = FApp::GetCurrentTime() - CachedSettings.BufferSettings.EngineTimeOffset - CachedSettings.BufferSettings.ValidEngineTime;
 			int32 FrameIndex = 0;
 			for (const FLiveLinkFrameDataStruct& SourceFrameData : FrameData)
 			{
-				FFrameTime UsedFrameTime = CurrentFrameTimeInFrameSpace - FFrameTime::FromDecimal(CachedSettings.BufferSettings.TimecodeFrameOffset) - CachedSettings.BufferSettings.ValidTimecodeFrame;
-				FFrameTime FrameTime = SourceFrameData.GetBaseData()->MetaData.SceneTime.Time;
-				if (FrameTime > UsedFrameTime)
+				double FrameTime = SourceFrameData.GetBaseData()->WorldTime.GetOffsettedTime();
+				double OffsetTime = ValidEngineTime;
+				if (FrameTime > OffsetTime)
 				{
 					break;
 				}
@@ -120,6 +95,37 @@ void FLiveLinkSubject::Update()
 				}
 			}
 		}
+	}
+	else if (GetMode() == ELiveLinkSourceMode::Timecode)
+	{
+		if (CachedSettings.BufferSettings.bValidTimecodeFrameEnabled)
+		{
+			if (FApp::GetCurrentFrameTime().IsSet())
+			{
+				const FQualifiedFrameTime CurrentSyncTime = FApp::GetCurrentFrameTime().GetValue();
+				const FFrameTime CurrentFrameTimeInFrameSpace = CurrentSyncTime.ConvertTo(CachedSettings.BufferSettings.TimecodeFrameRate);
+				int32 FrameIndex = 0;
+				for (const FLiveLinkFrameDataStruct& SourceFrameData : FrameData)
+				{
+					FFrameTime UsedFrameTime = CurrentFrameTimeInFrameSpace - FFrameTime::FromDecimal(CachedSettings.BufferSettings.TimecodeFrameOffset) - CachedSettings.BufferSettings.ValidTimecodeFrame;
+					FFrameTime FrameTime = SourceFrameData.GetBaseData()->MetaData.SceneTime.Time;
+					if (FrameTime > UsedFrameTime)
+					{
+						break;
+					}
+					++FrameIndex;
+				}
+
+				if (FrameIndex - 1 >= 0)
+				{
+					const int32 Count = CachedSettings.BufferSettings.bKeepAtLeastOneFrame && FrameData.Num() == FrameIndex ? FrameIndex - 1 : FrameIndex;
+					if (Count > 0)
+					{
+						FrameData.RemoveAt(0, Count, false);
+					}
+				}
+			}
+		}
 
 		// no warning if GetCurrentFrameTime is not set, the warning is done below after GetFrameAtSceneTime
 	}
@@ -369,12 +375,15 @@ void FLiveLinkSubject::AddFrameData(FLiveLinkFrameDataStruct&& InFrameData)
 
 int32 FLiveLinkSubject::FindNewFrame_WorldTime(const FLiveLinkWorldTime& WorldTime) const
 {
-	const double ValidEngineTime = FApp::GetCurrentTime() - CachedSettings.BufferSettings.EngineTimeOffset - CachedSettings.BufferSettings.ValidEngineTime;
-	const double WorldOffsettedTime = WorldTime.GetOffsettedTime();
-	if (WorldOffsettedTime < ValidEngineTime)
+	if (CachedSettings.BufferSettings.bValidEngineTimeEnabled)
 	{
-		static const FName NAME_InvalidWorldTime = "LiveLinkSubject_InvalidWorldTIme";
-		FLiveLinkLog::WarningOnce(NAME_InvalidWorldTime, SubjectKey, TEXT("Trying to add a frame in which the world time has a value too low compare to the engine's time. Do you have an invalid offset? The Subject is '%s'."), *SubjectKey.SubjectName.ToString());
+		const double ValidEngineTime = FApp::GetCurrentTime() - CachedSettings.BufferSettings.EngineTimeOffset - CachedSettings.BufferSettings.ValidEngineTime;
+		const double WorldOffsettedTime = WorldTime.GetOffsettedTime();
+		if (WorldOffsettedTime < ValidEngineTime)
+		{
+			static const FName NAME_InvalidWorldTime = "LiveLinkSubject_InvalidWorldTIme";
+			FLiveLinkLog::WarningOnce(NAME_InvalidWorldTime, SubjectKey, TEXT("Trying to add a frame in which the world time has a value too low compare to the engine's time. Do you have an invalid offset? The Subject is '%s'."), *SubjectKey.SubjectName.ToString());
+		}
 	}
 
 	return FindNewFrame_WorldTimeInternal(WorldTime);
@@ -418,7 +427,7 @@ int32 FLiveLinkSubject::FindNewFrame_SceneTime(const FQualifiedFrameTime& Qualif
 	}
 
 	// If we do not have a TC set, keep buffering, the TC may be unresponsive for a moment. We do not want to loose data.
-	if (FApp::GetCurrentFrameTime().IsSet())
+	if (FApp::GetCurrentFrameTime().IsSet() && CachedSettings.BufferSettings.bValidTimecodeFrameEnabled)
 	{
 		const FQualifiedFrameTime CurrentSyncTime = FApp::GetCurrentFrameTime().GetValue();
 		const FFrameTime CurrentFrameTimeInFrameSpace = CurrentSyncTime.ConvertTo(CachedSettings.BufferSettings.TimecodeFrameRate);
diff --git a/Engine/Plugins/Animation/LiveLink/Source/LiveLink/Private/LiveLinkTimedDataInput.h b/Engine/Plugins/Animation/LiveLink/Source/LiveLink/Private/LiveLinkTimedDataInput.h
--- a/Engine/Plugins/Animation/LiveLink/Source/LiveLink/Private/LiveLinkTimedDataInput.h
+++ b/Engine/Plugins/Animation/LiveLink/Source/LiveLink/Private/LiveLinkTimedDataInput.h
@@ -27,13 +27,13 @@ public:
 	virtual int32 GetDataBufferSize() const override;
 	virtual void SetDataBufferSize(int32 BufferSize) override;
 	virtual bool IsDataBufferSizeControlledByInput() const override { return true; }
+	virtual void AddChannel(ITimedDataInputChannel* Channel) override { Channels.Add(Channel); }
+	virtual void RemoveChannel(ITimedDataInputChannel* Channel) override { Channels.RemoveSingleSwap(Channel); }
 #if WITH_EDITOR
 	virtual const FSlateBrush* GetDisplayIcon() const override;
 #endif
 	//~ End ITimedDataInput API
 
-	void AddChannel(ITimedDataInputChannel* Channel) { Channels.Add(Channel); }
-	void RemoveChannel(ITimedDataInputChannel* Channel) { Channels.RemoveSingleSwap(Channel); }
 
 private:
 	FLiveLinkClient* LiveLinkClient;
diff --git a/Engine/Plugins/Animation/LiveLink/Source/LiveLink/Public/LiveLinkSettings.h b/Engine/Plugins/Animation/LiveLink/Source/LiveLink/Public/LiveLinkSettings.h
--- a/Engine/Plugins/Animation/LiveLink/Source/LiveLink/Public/LiveLinkSettings.h
+++ b/Engine/Plugins/Animation/LiveLink/Source/LiveLink/Public/LiveLinkSettings.h
@@ -7,9 +7,10 @@
 #include "UObject/ObjectMacros.h"
 
 #include "Engine/EngineTypes.h"
-#include "LiveLinkRole.h"
 #include "LiveLinkFrameInterpolationProcessor.h"
 #include "LiveLinkFramePreProcessor.h"
+#include "LiveLinkRole.h"
+#include "LiveLinkSourceSettings.h"
 #include "LiveLinkSubjectSettings.h"
 #include "Templates/SubclassOf.h"
 
@@ -76,6 +77,10 @@ public:
 	UPROPERTY(config, EditAnywhere, Category="LiveLink", meta=(DisplayName="Preset Save Location"))
 	FDirectoryPath PresetSaveDir;
 
+	/** The default evaluation mode a source connected via the message bus should start with. */
+	UPROPERTY(config, EditAnywhere, AdvancedDisplay, Category = "LiveLink")
+	ELiveLinkSourceMode DefaultMessageBusSourceMode;
+
 	/** The refresh frequency of the list of message bus provider (when discovery is requested). */
 	UPROPERTY(config, EditAnywhere, AdvancedDisplay, Category="LiveLink", meta=(ConfigRestartRequired=true, ForceUnits=s))
 	double MessageBusPingRequestFrequency;
diff --git a/Engine/Plugins/Media/AjaMedia/Source/AjaMedia/AjaMedia.Build.cs b/Engine/Plugins/Media/AjaMedia/Source/AjaMedia/AjaMedia.Build.cs
--- a/Engine/Plugins/Media/AjaMedia/Source/AjaMedia/AjaMedia.Build.cs
+++ b/Engine/Plugins/Media/AjaMedia/Source/AjaMedia/AjaMedia.Build.cs
@@ -28,6 +28,7 @@ namespace UnrealBuildTool.Rules
 					"Projects",
 					"TimeManagement",
 					"RenderCore",
+					"SlateCore"
 				});
 
 			PrivateIncludePathModuleNames.AddRange(
diff --git a/Engine/Plugins/Media/AjaMedia/Source/AjaMedia/Private/AjaMediaModule.cpp b/Engine/Plugins/Media/AjaMedia/Source/AjaMedia/Private/AjaMediaModule.cpp
--- a/Engine/Plugins/Media/AjaMedia/Source/AjaMedia/Private/AjaMediaModule.cpp
+++ b/Engine/Plugins/Media/AjaMedia/Source/AjaMedia/Private/AjaMediaModule.cpp
@@ -5,11 +5,14 @@
 #include "Aja/Aja.h"
 #include "AjaDeviceProvider.h"
 #include "AJALib.h"
-#include "Player/AjaMediaPlayer.h"
-
+#include "Brushes/SlateImageBrush.h"
 #include "CoreMinimal.h"
 #include "IMediaIOCoreModule.h"
+#include "Interfaces/IPluginManager.h"
 #include "Modules/ModuleManager.h"
+#include "Player/AjaMediaPlayer.h"
+#include "Styling/SlateStyle.h"
+
 
 DEFINE_LOG_CATEGORY(LogAjaMedia);
 
@@ -33,6 +36,11 @@ public:
 		return MakeShared<FAjaMediaPlayer, ESPMode::ThreadSafe>(EventSink);
 	}
 
+	virtual TSharedPtr<FSlateStyleSet> GetStyle() override
+	{
+		return StyleSet;
+	}
+
 	virtual bool IsInitialized() const override { return FAja::IsInitialized(); }
 
 	virtual bool CanBeUsed() const override { return FAja::CanUseAJACard(); }
@@ -49,6 +57,8 @@ public:
 			return;
 		}
 
+		CreateStyle();
+
 		IMediaIOCoreModule::Get().RegisterDeviceProvider(&DeviceProvider);
 	}
 
@@ -61,8 +71,20 @@ public:
 		FAja::Shutdown();
 	}
 
+private:
+	void CreateStyle()
+	{
+		static FName StyleName(TEXT("AjaMediaStyle"));
+		StyleSet = MakeShared<FSlateStyleSet>(StyleName);
+
+		const FString ContentDir = IPluginManager::Get().FindPlugin(TEXT("AjaMedia"))->GetContentDir();
+		const FVector2D Icon16x16(16.0f, 16.0f);
+		StyleSet->Set("AjaMediaIcon", new FSlateImageBrush((ContentDir / TEXT("Editor/Icons/AjaMediaSource_64x")) + TEXT(".png"), Icon16x16));
+	}
+
 private:
 	FAjaDeviceProvider DeviceProvider;
+	TSharedPtr<FSlateStyleSet> StyleSet;
 };
 
 IMPLEMENT_MODULE(FAjaMediaModule, AjaMedia);
diff --git a/Engine/Plugins/Media/AjaMedia/Source/AjaMedia/Private/Player/AjaMediaPlayer.cpp b/Engine/Plugins/Media/AjaMedia/Source/AjaMedia/Private/Player/AjaMediaPlayer.cpp
--- a/Engine/Plugins/Media/AjaMedia/Source/AjaMedia/Private/Player/AjaMediaPlayer.cpp
+++ b/Engine/Plugins/Media/AjaMedia/Source/AjaMedia/Private/Player/AjaMediaPlayer.cpp
@@ -10,10 +10,12 @@
 
 #include "HAL/PlatformAtomics.h"
 #include "HAL/PlatformProcess.h"
+#include "IAjaMediaModule.h"
 #include "IMediaEventSink.h"
 #include "IMediaOptions.h"
 #include "Misc/ScopeLock.h"
 #include "Stats/Stats2.h"
+#include "Styling/SlateStyle.h"
 
 #include "AjaMediaAudioSample.h"
 #include "AjaMediaBinarySample.h"
@@ -56,9 +58,6 @@ FAjaMediaPlayer::FAjaMediaPlayer(IMediaEventSink& InEventSink)
 	, AjaThreadAudioChannels(0)
 	, AjaThreadAudioSampleRate(0)
 	, AjaThreadFrameDropCount(0)
-	, AjaThreadAutoCirculateAudioFrameDropCount(0)
-	, AjaThreadAutoCirculateMetadataFrameDropCount(0)
-	, AjaThreadAutoCirculateVideoFrameDropCount(0)
 	, LastFrameDropCount(0)
 	, PreviousFrameDropCount(0)
 	, bEncodeTimecodeInTexel(false)
@@ -69,7 +68,10 @@ FAjaMediaPlayer::FAjaMediaPlayer(IMediaEventSink& InEventSink)
 	, bUseVideo(false)
 	, bVerifyFrameDropCount(true)
 	, InputChannel(nullptr)
-{ }
+	, SupportedSampleTypes(static_cast<EMediaIOSampleType>(static_cast<uint8>(EMediaIOSampleType::Video) | static_cast<uint8>(EMediaIOSampleType::Audio) | static_cast<uint8>(EMediaIOSampleType::Metadata)))
+{
+	Samples->EnableTimedDataChannels(SupportedSampleTypes);
+}
 
 
 FAjaMediaPlayer::~FAjaMediaPlayer()
@@ -206,6 +208,9 @@ bool FAjaMediaPlayer::Open(const FString& Url, const IMediaOptions* Options)
 		InputChannel = nullptr;
 	}
 
+	// Setup our different supported channels based on source settings
+	SetupSampleChannels();
+
 	// configure format information for base class
 	AudioTrackFormat.BitsPerSample = 32;
 	AudioTrackFormat.NumChannels = 0;
@@ -295,9 +300,16 @@ FString FAjaMediaPlayer::GetStats() const
 	return Stats;
 }
 
+#if WITH_EDITOR
+const FSlateBrush* FAjaMediaPlayer::GetDisplayIcon() const
+{
+	return IAjaMediaModule::Get().GetStyle()->GetBrush("AjaMediaIcon");
+}
+#endif //WITH_EDITOR
 
-void FAjaMediaPlayer::TickFetch(FTimespan DeltaTime, FTimespan /*Timecode*/)
+void FAjaMediaPlayer::TickFetch(FTimespan DeltaTime, FTimespan Timecode)
 {
+	Super::TickFetch(DeltaTime, Timecode);
 	if (InputChannel && CurrentState == EMediaState::Playing)
 	{
 		ProcessFrame();
@@ -305,7 +317,6 @@ void FAjaMediaPlayer::TickFetch(FTimespan DeltaTime, FTimespan /*Timecode*/)
 	}
 }
 
-
 void FAjaMediaPlayer::TickInput(FTimespan DeltaTime, FTimespan Timecode)
 {
 	// update player state
@@ -350,29 +361,7 @@ void FAjaMediaPlayer::ProcessFrame()
 
 void FAjaMediaPlayer::VerifyFrameDropCount()
 {
-	//Verify if a buffer is in overflow state. Popping samples MUST be done from the GameThread to respect single consumer
-
-	//Anc buffer
-	int32 MetaDataOverflowCount = FMath::Max(Samples->NumMetadataSamples() - MaxNumMetadataFrameBuffer, 0);
-	for (int32 i = 0; i < MetaDataOverflowCount; ++i)
-	{
-		Samples->PopMetadata();
-	}
-
-	//Audio buffer
-	int32 AudioOverflowCount = FMath::Max(Samples->NumAudioSamples() - MaxNumAudioFrameBuffer, 0);
-	for (int32 i = 0; i < AudioOverflowCount; ++i)
-	{
-		Samples->PopAudio();
-	}
-
-	//Video buffer
-	int32 VideoOverflowCount = FMath::Max(Samples->NumVideoSamples() - MaxNumVideoFrameBuffer, 0);
-	for (int32 i = 0; i < VideoOverflowCount; ++i)
-	{
-		Samples->PopVideo();
-	}
-
+	//Update frame dropped status
 	if (bVerifyFrameDropCount)
 	{
 		uint32 FrameDropCount = AjaThreadFrameDropCount;
@@ -393,22 +382,40 @@ void FAjaMediaPlayer::VerifyFrameDropCount()
 		}
 		LastFrameDropCount = FrameDropCount;
 
-		MetaDataOverflowCount += FPlatformAtomics::InterlockedExchange(&AjaThreadAutoCirculateMetadataFrameDropCount, 0);
-		if (MetaDataOverflowCount > 0)
+		const int32 CurrentMetadataDropCount = Samples->GetMetadataFrameDropCount();
+		int32 DeltaMetadataDropCount = CurrentMetadataDropCount;
+		if (CurrentMetadataDropCount >= PreviousMetadataFrameDropCount)
+		{
+			DeltaMetadataDropCount = CurrentMetadataDropCount - PreviousMetadataFrameDropCount;
+		}
+		PreviousMetadataFrameDropCount = CurrentMetadataDropCount;
+		if (DeltaMetadataDropCount > 0)
 		{
-			UE_LOG(LogAjaMedia, Warning, TEXT("Lost %d metadata frames on input %s. Frame rate is either too slow or buffering capacity is too small."), MetaDataOverflowCount, *GetUrl());
+			UE_LOG(LogAjaMedia, Warning, TEXT("Lost %d metadata frames on input %s. Frame rate is either too slow or buffering capacity is too small."), DeltaMetadataDropCount, *GetUrl());
 		}
 
-		AudioOverflowCount += FPlatformAtomics::InterlockedExchange(&AjaThreadAutoCirculateAudioFrameDropCount, 0);
-		if (AudioOverflowCount > 0)
+		const int32 CurrentAudioDropCount = Samples->GetAudioFrameDropCount();
+		int32 DeltaAudioDropCount = CurrentAudioDropCount;
+		if (CurrentAudioDropCount >= PreviousAudioFrameDropCount)
 		{
-			UE_LOG(LogAjaMedia, Warning, TEXT("Lost %d audio frames on input %s. Frame rate is either too slow or buffering capacity is too small."), AudioOverflowCount, *GetUrl());
+			DeltaAudioDropCount = CurrentAudioDropCount - PreviousAudioFrameDropCount;
+		}
+		PreviousAudioFrameDropCount = CurrentAudioDropCount;
+		if (DeltaAudioDropCount > 0)
+		{
+			UE_LOG(LogAjaMedia, Warning, TEXT("Lost %d audio frames on input %s. Frame rate is either too slow or buffering capacity is too small."), DeltaAudioDropCount, *GetUrl());
 		}
 
-		VideoOverflowCount += FPlatformAtomics::InterlockedExchange(&AjaThreadAutoCirculateVideoFrameDropCount, 0);
-		if (bVerifyFrameDropCount && VideoOverflowCount > 0)
+		const int32 CurrentVideoDropCount = Samples->GetVideoFrameDropCount();
+		int32 DeltaVideoDropCount = CurrentVideoDropCount;
+		if (CurrentVideoDropCount >= PreviousVideoFrameDropCount)
+		{
+			DeltaVideoDropCount = CurrentVideoDropCount - PreviousVideoFrameDropCount;
+		}
+		PreviousVideoFrameDropCount = CurrentVideoDropCount;
+		if (DeltaVideoDropCount > 0)
 		{
-			UE_LOG(LogAjaMedia, Warning, TEXT("Lost %d video frames on input %s. Frame rate is either too slow or buffering capacity is too small."), VideoOverflowCount, *GetUrl());
+			UE_LOG(LogAjaMedia, Warning, TEXT("Lost %d video frames on input %s. Frame rate is either too slow or buffering capacity is too small."), DeltaVideoDropCount, *GetUrl());
 		}
 	}
 }
@@ -448,73 +455,29 @@ bool FAjaMediaPlayer::OnRequestInputBuffer(const AJA::AJARequestInputBufferData&
 	// Anc Field 1
 	if (bUseAncillary && InRequestBuffer.AncBufferSize > 0)
 	{
-		const int32 NumMetadataSamples = Samples->NumMetadataSamples();
-		if (NumMetadataSamples >= MaxNumMetadataFrameBuffer * AjaMediaPlayerConst::ToleratedExtraMaxBufferCount)
-		{
-			if (bVerifyFrameDropCount)
-			{
-				FPlatformAtomics::InterlockedIncrement(&AjaThreadAutoCirculateMetadataFrameDropCount);
-			}
-		}
-		else
-		{
-			AjaThreadCurrentAncSample = MetadataSamplePool->AcquireShared();
-			OutRequestedBuffer.AncBuffer = reinterpret_cast<uint8_t*>(AjaThreadCurrentAncSample->RequestBuffer(InRequestBuffer.AncBufferSize));
-		}
+		AjaThreadCurrentAncSample = MetadataSamplePool->AcquireShared();
+		OutRequestedBuffer.AncBuffer = reinterpret_cast<uint8_t*>(AjaThreadCurrentAncSample->RequestBuffer(InRequestBuffer.AncBufferSize));
 	}
 
 	// Anc Field 2
 	if (bUseAncillary && InRequestBuffer.AncF2BufferSize > 0)
 	{
-		const int32 NumMetadataSamples = Samples->NumMetadataSamples();
-		if (NumMetadataSamples >= MaxNumMetadataFrameBuffer * AjaMediaPlayerConst::ToleratedExtraMaxBufferCount)
-		{
-			if (bVerifyFrameDropCount)
-			{
-				FPlatformAtomics::InterlockedIncrement(&AjaThreadAutoCirculateMetadataFrameDropCount);
-			}
-		}
-		else
-		{
-			AjaThreadCurrentAncF2Sample = MetadataSamplePool->AcquireShared();
-			OutRequestedBuffer.AncBuffer = reinterpret_cast<uint8_t*>(AjaThreadCurrentAncF2Sample->RequestBuffer(InRequestBuffer.AncF2BufferSize));
-		}
+		AjaThreadCurrentAncF2Sample = MetadataSamplePool->AcquireShared();
+		OutRequestedBuffer.AncBuffer = reinterpret_cast<uint8_t*>(AjaThreadCurrentAncF2Sample->RequestBuffer(InRequestBuffer.AncF2BufferSize));
 	}
 
 	// Audio
 	if (bUseAudio && InRequestBuffer.AudioBufferSize > 0)
 	{
-		const int32 NumAudioSamples = Samples->NumAudioSamples();
-		if (NumAudioSamples >= MaxNumAudioFrameBuffer * AjaMediaPlayerConst::ToleratedExtraMaxBufferCount)
-		{
-			if (bVerifyFrameDropCount)
-			{
-				FPlatformAtomics::InterlockedIncrement(&AjaThreadAutoCirculateAudioFrameDropCount);
-			}
-		}
-		else
-		{
-			AjaThreadCurrentAudioSample = AudioSamplePool->AcquireShared();
-			OutRequestedBuffer.AudioBuffer = reinterpret_cast<uint8_t*>(AjaThreadCurrentAudioSample->RequestBuffer(InRequestBuffer.AudioBufferSize));
-		}
+		AjaThreadCurrentAudioSample = AudioSamplePool->AcquireShared();
+		OutRequestedBuffer.AudioBuffer = reinterpret_cast<uint8_t*>(AjaThreadCurrentAudioSample->RequestBuffer(InRequestBuffer.AudioBufferSize));
 	}
 
 	// Video
 	if (bUseVideo && InRequestBuffer.VideoBufferSize > 0 && InRequestBuffer.bIsProgressivePicture)
 	{
-		const int32 NumVideoSamples = Samples->NumVideoSamples();
-		if (NumVideoSamples >= MaxNumVideoFrameBuffer * AjaMediaPlayerConst::ToleratedExtraMaxBufferCount)
-		{
-			if (bVerifyFrameDropCount)
-			{
-				FPlatformAtomics::InterlockedIncrement(&AjaThreadAutoCirculateVideoFrameDropCount);
-			}
-		}
-		else
-		{
-			AjaThreadCurrentTextureSample = TextureSamplePool->AcquireShared();
-			OutRequestedBuffer.VideoBuffer = reinterpret_cast<uint8_t*>(AjaThreadCurrentTextureSample->RequestBuffer(InRequestBuffer.VideoBufferSize));
-		}
+		AjaThreadCurrentTextureSample = TextureSamplePool->AcquireShared();
+		OutRequestedBuffer.VideoBuffer = reinterpret_cast<uint8_t*>(AjaThreadCurrentTextureSample->RequestBuffer(InRequestBuffer.VideoBufferSize));
 	}
 
 	return true;
@@ -578,18 +541,10 @@ bool FAjaMediaPlayer::OnInputFrameReceived(const AJA::AJAInputFrameData& InInput
 		}
 		else
 		{
-			const int32 NumMetadataSamples = Samples->NumMetadataSamples();
-			if (NumMetadataSamples >= MaxNumMetadataFrameBuffer * AjaMediaPlayerConst::ToleratedExtraMaxBufferCount)
+			auto MetaDataSample = MetadataSamplePool->AcquireShared();
+			if (MetaDataSample->Initialize(InAncillaryFrame.AncBuffer, InAncillaryFrame.AncBufferSize, DecodedTime, VideoFrameRate, DecodedTimecode))
 			{
-				FPlatformAtomics::InterlockedIncrement(&AjaThreadAutoCirculateMetadataFrameDropCount);
-			}
-			else
-			{
-				auto MetaDataSample = MetadataSamplePool->AcquireShared();
-				if (MetaDataSample->Initialize(InAncillaryFrame.AncBuffer, InAncillaryFrame.AncBufferSize, DecodedTime, VideoFrameRate, DecodedTimecode))
-				{
-					Samples->AddMetadata(MetaDataSample);
-				}
+				Samples->AddMetadata(MetaDataSample);
 			}
 		}
 	}
@@ -606,18 +561,10 @@ bool FAjaMediaPlayer::OnInputFrameReceived(const AJA::AJAInputFrameData& InInput
 		}
 		else
 		{
-			const int32 NumMetadataSamples = Samples->NumMetadataSamples();
-			if (NumMetadataSamples >= MaxNumMetadataFrameBuffer * AjaMediaPlayerConst::ToleratedExtraMaxBufferCount)
+			auto MetaDataSample = MetadataSamplePool->AcquireShared();
+			if (MetaDataSample->Initialize(InAncillaryFrame.AncF2Buffer, InAncillaryFrame.AncF2BufferSize, DecodedTimeF2, VideoFrameRate, DecodedTimecodeF2))
 			{
-				FPlatformAtomics::InterlockedIncrement(&AjaThreadAutoCirculateMetadataFrameDropCount);
-			}
-			else
-			{
-				auto MetaDataSample = MetadataSamplePool->AcquireShared();
-				if (MetaDataSample->Initialize(InAncillaryFrame.AncF2Buffer, InAncillaryFrame.AncF2BufferSize, DecodedTimeF2, VideoFrameRate, DecodedTimecodeF2))
-				{
-					Samples->AddMetadata(MetaDataSample);
-				}
+				Samples->AddMetadata(MetaDataSample);
 			}
 		}
 	}
@@ -637,21 +584,14 @@ bool FAjaMediaPlayer::OnInputFrameReceived(const AJA::AJAInputFrameData& InInput
 		}
 		else
 		{
-			if (Samples->NumAudioSamples() >= MaxNumAudioFrameBuffer * AjaMediaPlayerConst::ToleratedExtraMaxBufferCount)
+			auto AudioSample = AudioSamplePool->AcquireShared();
+			if (AudioSample->Initialize(InAudioFrame, DecodedTime, DecodedTimecode))
 			{
-				FPlatformAtomics::InterlockedIncrement(&AjaThreadAutoCirculateAudioFrameDropCount);
+				Samples->AddAudio(AudioSample);
 			}
-			else
-			{
-				auto AudioSample = AudioSamplePool->AcquireShared();
-				if (AudioSample->Initialize(InAudioFrame, DecodedTime, DecodedTimecode))
-				{
-					Samples->AddAudio(AudioSample);
-				}
 
-				AjaThreadAudioChannels = AudioSample->GetChannels();
-				AjaThreadAudioSampleRate = AudioSample->GetSampleRate();
-			}
+			AjaThreadAudioChannels = AudioSample->GetChannels();
+			AjaThreadAudioSampleRate = AudioSample->GetSampleRate();
 		}
 	}
 
@@ -708,35 +648,27 @@ bool FAjaMediaPlayer::OnInputFrameReceived(const AJA::AJAInputFrameData& InInput
 		}
 		else
 		{
-			const int32 NumVideoSamples = Samples->NumVideoSamples() + (!InVideoFrame.bIsProgressivePicture ? 1 : 0);
-			if (NumVideoSamples >= MaxNumVideoFrameBuffer * AjaMediaPlayerConst::ToleratedExtraMaxBufferCount)
+			auto TextureSample = TextureSamplePool->AcquireShared();
+			if (InVideoFrame.bIsProgressivePicture)
 			{
-				FPlatformAtomics::InterlockedIncrement(&AjaThreadAutoCirculateVideoFrameDropCount);
+				if (TextureSample->InitializeProgressive(InVideoFrame, VideoSampleFormat, DecodedTime, VideoFrameRate, DecodedTimecode, bIsSRGBInput))
+				{
+					Samples->AddVideo(TextureSample);
+				}
 			}
 			else
 			{
-				auto TextureSample = TextureSamplePool->AcquireShared();
-				if (InVideoFrame.bIsProgressivePicture)
+				bool bEven = true;
+				if (TextureSample->InitializeInterlaced_Halfed(InVideoFrame, VideoSampleFormat, DecodedTime, VideoFrameRate, DecodedTimecode, bEven, bIsSRGBInput))
 				{
-					if (TextureSample->InitializeProgressive(InVideoFrame, VideoSampleFormat, DecodedTime, VideoFrameRate, DecodedTimecode, bIsSRGBInput))
-					{
-						Samples->AddVideo(TextureSample);
-					}
+					Samples->AddVideo(TextureSample);
 				}
-				else
+
+				auto TextureSampleOdd = TextureSamplePool->AcquireShared();
+				bEven = false;
+				if (TextureSampleOdd->InitializeInterlaced_Halfed(InVideoFrame, VideoSampleFormat, DecodedTimeF2, VideoFrameRate, DecodedTimecodeF2, bEven, bIsSRGBInput))
 				{
-					bool bEven = true;
-					if (TextureSample->InitializeInterlaced_Halfed(InVideoFrame, VideoSampleFormat, DecodedTime, VideoFrameRate, DecodedTimecode, bEven, bIsSRGBInput))
-					{
-						Samples->AddVideo(TextureSample);
-					}
-
-					auto TextureSampleOdd = TextureSamplePool->AcquireShared();
-					bEven = false;
-					if (TextureSampleOdd->InitializeInterlaced_Halfed(InVideoFrame, VideoSampleFormat, DecodedTimeF2, VideoFrameRate, DecodedTimecodeF2, bEven, bIsSRGBInput))
-					{
-						Samples->AddVideo(TextureSampleOdd);
-					}
+					Samples->AddVideo(TextureSampleOdd);
 				}
 			}
 		}
@@ -763,6 +695,21 @@ bool FAjaMediaPlayer::IsHardwareReady() const
 	return AjaThreadNewState == EMediaState::Playing ? true : false;
 }
 
+void FAjaMediaPlayer::SetupSampleChannels()
+{
+	FMediaIOSamplingSettings VideoSettings = BaseSettings;
+	VideoSettings.BufferSize = MaxNumVideoFrameBuffer;
+	Samples->InitializeVideoBuffer(VideoSettings);
+
+	FMediaIOSamplingSettings AudioSettings = BaseSettings;
+	AudioSettings.BufferSize = MaxNumAudioFrameBuffer;
+	Samples->InitializeAudioBuffer(AudioSettings);
+
+	FMediaIOSamplingSettings MetadataSettings = BaseSettings;
+	MetadataSettings.BufferSize = MaxNumMetadataFrameBuffer;
+	Samples->InitializeMetadataBuffer(MetadataSettings);
+}
+
 #undef LOCTEXT_NAMESPACE
 
 #include "AjaMediaHidePlatformTypes.h"
diff --git a/Engine/Plugins/Media/AjaMedia/Source/AjaMedia/Private/Player/AjaMediaPlayer.h b/Engine/Plugins/Media/AjaMedia/Source/AjaMedia/Private/Player/AjaMediaPlayer.h
--- a/Engine/Plugins/Media/AjaMedia/Source/AjaMedia/Private/Player/AjaMediaPlayer.h
+++ b/Engine/Plugins/Media/AjaMedia/Source/AjaMedia/Private/Player/AjaMediaPlayer.h
@@ -13,9 +13,11 @@ class FAjaMediaBinarySamplePool;
 class FAjaMediaTextureSample;
 class FAjaMediaTextureSamplePool;
 class FMediaIOCoreBinarySampleBase;
+struct FSlateBrush;
 class IMediaEventSink;
 
 enum class EMediaTextureSampleFormat;
+enum class EMediaIOSampleType;
 
 namespace AJA
 {
@@ -64,6 +66,11 @@ public:
 
 	virtual FString GetStats() const override;
 
+	//~ ITimedDataInput interface
+#if WITH_EDITOR
+	virtual const FSlateBrush* GetDisplayIcon() const override;
+#endif
+
 protected:
 
 	//~ IAJAInputOutputCallbackInterface interface
@@ -89,6 +96,9 @@ protected:
 
 	virtual bool IsHardwareReady() const override;
 
+	//~ FMediaIOCorePlayerBase interface
+	virtual void SetupSampleChannels() override;
+
 private:
 
 	/** Audio, MetaData, Texture  sample object pool. */
@@ -120,9 +130,9 @@ private:
 
 	/** Number of frames drop from the last tick. */
 	int32 AjaThreadFrameDropCount;
-	int32 AjaThreadAutoCirculateAudioFrameDropCount;
-	int32 AjaThreadAutoCirculateMetadataFrameDropCount;
-	int32 AjaThreadAutoCirculateVideoFrameDropCount;
+	int32 PreviousMetadataFrameDropCount;
+	int32 PreviousAudioFrameDropCount;
+	int32 PreviousVideoFrameDropCount;
 
 	/** Number of frames drop from the last tick. */
 	uint32 LastFrameDropCount;
@@ -146,6 +156,9 @@ private:
 	/** Maps to the current input Device */
 	AJA::AJAInputChannel* InputChannel;
 
+	/** Used to flag which sample types we advertise as supported for timed data monitoring */
+	EMediaIOSampleType SupportedSampleTypes;
+
 	/** Frame Description from capture device */
 	AJA::FAJAVideoFormat LastVideoFormatIndex;
 	/** Previous frame timecode for stats purpose */
diff --git a/Engine/Plugins/Media/AjaMedia/Source/AjaMedia/Public/IAjaMediaModule.h b/Engine/Plugins/Media/AjaMedia/Source/AjaMedia/Public/IAjaMediaModule.h
--- a/Engine/Plugins/Media/AjaMedia/Source/AjaMedia/Public/IAjaMediaModule.h
+++ b/Engine/Plugins/Media/AjaMedia/Source/AjaMedia/Public/IAjaMediaModule.h
@@ -3,8 +3,10 @@
 #pragma once
 
 #include "Modules/ModuleInterface.h"
+#include "Modules/ModuleManager.h"
 #include "Templates/SharedPointer.h"
 
+class FSlateStyleSet;
 class IMediaEventSink;
 class IMediaPlayer;
 
@@ -15,6 +17,12 @@ class IAjaMediaModule : public IModuleInterface
 {
 public:
 
+	static inline IAjaMediaModule& Get()
+	{
+		static const FName ModuleName = "AjaMedia";
+		return FModuleManager::LoadModuleChecked<IAjaMediaModule>(ModuleName);
+	}
+
 	/**
 	 * Create an AJA based media player.
 	 * @param EventSink The object that receives media events from the player.
@@ -22,6 +30,9 @@ public:
 	 */
 	virtual TSharedPtr<IMediaPlayer, ESPMode::ThreadSafe> CreatePlayer(IMediaEventSink& EventSink) = 0;
 
+	/** @return SlateStyleSet to be used across the AjaMedia module */
+	virtual TSharedPtr<FSlateStyleSet> GetStyle() = 0;
+
 	/** @return true if the Aja module and AJA dll could be loaded */
 	virtual bool IsInitialized() const = 0;
 
diff --git a/Engine/Plugins/Media/BlackmagicMedia/Source/BlackmagicMedia/BlackmagicMedia.Build.cs b/Engine/Plugins/Media/BlackmagicMedia/Source/BlackmagicMedia/BlackmagicMedia.Build.cs
--- a/Engine/Plugins/Media/BlackmagicMedia/Source/BlackmagicMedia/BlackmagicMedia.Build.cs
+++ b/Engine/Plugins/Media/BlackmagicMedia/Source/BlackmagicMedia/BlackmagicMedia.Build.cs
@@ -32,7 +32,8 @@ namespace UnrealBuildTool.Rules
 					"Blackmagic",
 					"Engine",
 					"MediaUtils",
-					"Projects"
+					"Projects",
+					"SlateCore"
                 });
 
 			PrivateIncludePathModuleNames.AddRange(
diff --git a/Engine/Plugins/Media/BlackmagicMedia/Source/BlackmagicMedia/Private/BlackmagicMediaModule.cpp b/Engine/Plugins/Media/BlackmagicMedia/Source/BlackmagicMedia/Private/BlackmagicMediaModule.cpp
--- a/Engine/Plugins/Media/BlackmagicMedia/Source/BlackmagicMedia/Private/BlackmagicMediaModule.cpp
+++ b/Engine/Plugins/Media/BlackmagicMedia/Source/BlackmagicMedia/Private/BlackmagicMediaModule.cpp
@@ -6,9 +6,11 @@
 #include "BlackmagicDeviceProvider.h"
 #include "BlackmagicMediaPrivate.h"
 #include "BlackmagicMediaPlayer.h"
-
+#include "Brushes/SlateImageBrush.h"
 #include "IMediaIOCoreModule.h"
+#include "Interfaces/IPluginManager.h"
 #include "Modules/ModuleManager.h"
+#include "Styling/SlateStyle.h"
 
 
 DEFINE_LOG_CATEGORY(LogBlackmagicMedia);
@@ -33,6 +35,11 @@ public:
 		return MakeShared<FBlackmagicMediaPlayer, ESPMode::ThreadSafe>(EventSink);
 	}
 
+	virtual TSharedPtr<FSlateStyleSet> GetStyle() override
+	{
+		return StyleSet;
+	}
+
 	virtual bool IsInitialized() const override { return FBlackmagic::IsInitialized(); }
 
 	virtual bool CanBeUsed() const override { return FBlackmagic::CanUseBlackmagicCard(); }
@@ -49,6 +56,7 @@ public:
 			return;
 		}
 
+		CreateStyle();
 		
 		IMediaIOCoreModule::Get().RegisterDeviceProvider(&DeviceProvider);
 	}
@@ -63,8 +71,20 @@ public:
 		FBlackmagic::Shutdown();
 	}
 
+private:
+	void CreateStyle()
+	{
+		static FName StyleName(TEXT("BlackmagicMediaStyle"));
+		StyleSet = MakeShared<FSlateStyleSet>(StyleName);
+
+		const FString ContentDir = IPluginManager::Get().FindPlugin(TEXT("BlackmagicMedia"))->GetContentDir();
+		const FVector2D Icon16x16(16.0f, 16.0f);
+		StyleSet->Set("BlackmagicMediaIcon", new FSlateImageBrush((ContentDir / TEXT("Editor/Icons/BlackmagicMediaSource_64x")) + TEXT(".png"), Icon16x16));
+	}
+
 private:
 	FBlackmagicDeviceProvider DeviceProvider;
+	TSharedPtr<FSlateStyleSet> StyleSet;
 };
 
 IMPLEMENT_MODULE(FBlackmagicMediaModule, BlackmagicMedia);
diff --git a/Engine/Plugins/Media/BlackmagicMedia/Source/BlackmagicMedia/Private/Player/BlackmagicMediaPlayer.cpp b/Engine/Plugins/Media/BlackmagicMedia/Source/BlackmagicMedia/Private/Player/BlackmagicMediaPlayer.cpp
--- a/Engine/Plugins/Media/BlackmagicMedia/Source/BlackmagicMedia/Private/Player/BlackmagicMediaPlayer.cpp
+++ b/Engine/Plugins/Media/BlackmagicMedia/Source/BlackmagicMedia/Private/Player/BlackmagicMediaPlayer.cpp
@@ -5,6 +5,7 @@
 #include "Blackmagic.h"
 #include "BlackmagicMediaPrivate.h"
 #include "BlackmagicMediaSource.h"
+#include "IBlackmagicMediaModule.h"
 
 #include "HAL/CriticalSection.h"
 #include "HAL/PlatformProcess.h"
@@ -22,6 +23,7 @@
 #include "Misc/App.h"
 #include "Slate/SceneViewport.h"
 #include "Stats/Stats2.h"
+#include "Styling/SlateStyle.h"
 
 #include "BlackmagicMediaSource.h"
 
@@ -55,9 +57,8 @@ namespace BlackmagicMediaPlayerHelpers
 			, LastBitsPerSample(0)
 			, LastNumChannels(0)
 			, LastSampleRate(0)
-			, AudioFrameDropCount(0)
-			, MetadataFrameDropCount(0)
-			, VideoFrameDropCount(0)
+			, PreviousAudioFrameDropCount(0)
+			, PreviousVideoFrameDropCount(0)
 			, LastHasFrameTime(0.0)
 			, bReceivedValidFrame(false)
 			, bIsTimecodeExpected(false)
@@ -108,32 +109,30 @@ namespace BlackmagicMediaPlayerHelpers
 
 		void VerifyFrameDropCount_GameThread(const FString& InUrl)
 		{
-			//Audio buffer
-			int32 AudioOverflowCount = FMath::Max(MediaPlayer->Samples->NumAudioSamples() - MaxNumAudioFrameBuffer, 0);
-			for (int32 i = 0; i < AudioOverflowCount; ++i)
-			{
-				MediaPlayer->Samples->PopAudio();
-			}
-
-			//Video buffer
-			int32 VideoOverflowCount = FMath::Max(MediaPlayer->Samples->NumVideoSamples() - MaxNumVideoFrameBuffer, 0);
-			for (int32 i = 0; i < VideoOverflowCount; ++i)
-			{
-				MediaPlayer->Samples->PopVideo();
-			}
-
 			if (MediaPlayer->bVerifyFrameDropCount)
 			{
-				AudioOverflowCount += FPlatformAtomics::InterlockedExchange(&AudioFrameDropCount, 0);
-				if (AudioOverflowCount > 0)
+				const int32 CurrentAudioDropCount = MediaPlayer->Samples->GetAudioFrameDropCount();
+				int32 DeltaAudioDropCount = CurrentAudioDropCount;
+				if (CurrentAudioDropCount >= PreviousAudioFrameDropCount)
+				{
+					DeltaAudioDropCount = CurrentAudioDropCount - PreviousAudioFrameDropCount;
+				}
+				PreviousAudioFrameDropCount = CurrentAudioDropCount;
+				if (DeltaAudioDropCount > 0)
 				{
-					UE_LOG(LogBlackmagicMedia, Warning, TEXT("Lost %d audio frames on input %s. Frame rate is either too slow or buffering capacity is too small."), AudioOverflowCount, *InUrl);
+					UE_LOG(LogBlackmagicMedia, Warning, TEXT("Lost %d audio frames on input %s. Frame rate is either too slow or buffering capacity is too small."), DeltaAudioDropCount, *InUrl);
 				}
 
-				VideoOverflowCount += FPlatformAtomics::InterlockedExchange(&VideoFrameDropCount, 0);
-				if (VideoOverflowCount > 0)
+				const int32 CurrentVideoDropCount = MediaPlayer->Samples->GetVideoFrameDropCount();
+				int32 DeltaVideoDropCount = CurrentVideoDropCount;
+				if (CurrentVideoDropCount >= PreviousVideoFrameDropCount)
+				{
+					DeltaVideoDropCount = CurrentVideoDropCount - PreviousVideoFrameDropCount;
+				}
+				PreviousVideoFrameDropCount = CurrentVideoDropCount;
+				if (DeltaVideoDropCount > 0)
 				{
-					UE_LOG(LogBlackmagicMedia, Warning, TEXT("Lost %d video frames on input %s. Frame rate is either too slow or buffering capacity is too small."), VideoOverflowCount, *InUrl);
+					UE_LOG(LogBlackmagicMedia, Warning, TEXT("Lost %d video frames on input %s. Frame rate is either too slow or buffering capacity is too small."), DeltaVideoDropCount, *InUrl);
 				}
 			}
 		}
@@ -236,126 +235,105 @@ namespace BlackmagicMediaPlayerHelpers
 
 				if (InFrameInfo.AudioBuffer)
 				{
-					if (MediaPlayer->Samples->NumAudioSamples() >= MaxNumAudioFrameBuffer * BlackmagicMediaPlayerHelpers::ToleratedExtraMaxBufferCount)
+					auto AudioSamle = MediaPlayer->AudioSamplePool->AcquireShared();
+					if (AudioSamle->Initialize(reinterpret_cast<int32*>(InFrameInfo.AudioBuffer)
+						, InFrameInfo.AudioBufferSize / sizeof(int32)
+						, InFrameInfo.NumberOfAudioChannel
+						, InFrameInfo.AudioRate
+						, DecodedTime
+						, DecodedTimecode))
 					{
-						if (MediaPlayer->bVerifyFrameDropCount)
-						{
-							FPlatformAtomics::InterlockedIncrement(&AudioFrameDropCount);
-						}
-					}
-					else
-					{
-						auto AudioSamle = MediaPlayer->AudioSamplePool->AcquireShared();
-						if (AudioSamle->Initialize(reinterpret_cast<int32*>(InFrameInfo.AudioBuffer)
-							, InFrameInfo.AudioBufferSize / sizeof(int32)
-							, InFrameInfo.NumberOfAudioChannel
-							, InFrameInfo.AudioRate
-							, DecodedTime
-							, DecodedTimecode))
-						{
-							MediaPlayer->Samples->AddAudio(AudioSamle);
+						MediaPlayer->Samples->AddAudio(AudioSamle);
 
-							LastBitsPerSample = sizeof(int32);
-							LastSampleRate = InFrameInfo.AudioRate;
-							LastNumChannels = InFrameInfo.NumberOfAudioChannel;
-						}
+						LastBitsPerSample = sizeof(int32);
+						LastSampleRate = InFrameInfo.AudioRate;
+						LastNumChannels = InFrameInfo.NumberOfAudioChannel;
 					}
 				}
 
 				if (InFrameInfo.VideoBuffer)
 				{
 					const bool bIsProgressivePicture = InFrameInfo.FieldDominance != BlackmagicDesign::EFieldDominance::Interlaced;
-					const int32 NumVideoSamples = MediaPlayer->Samples->NumVideoSamples() + (!bIsProgressivePicture ? 1 : 0);
-					if (NumVideoSamples >= MaxNumVideoFrameBuffer * BlackmagicMediaPlayerHelpers::ToleratedExtraMaxBufferCount)
+					EMediaTextureSampleFormat SampleFormat = EMediaTextureSampleFormat::CharBGRA;
+					EMediaIOCoreEncodePixelFormat EncodePixelFormat = EMediaIOCoreEncodePixelFormat::CharUYVY;
+					FString OutputFilename = "";
+
+					switch (InFrameInfo.PixelFormat)
 					{
-						if (MediaPlayer->bVerifyFrameDropCount)
-						{
-							FPlatformAtomics::InterlockedIncrement(&VideoFrameDropCount);
-						}
+					case BlackmagicDesign::EPixelFormat::pf_8Bits:
+						SampleFormat = EMediaTextureSampleFormat::CharUYVY;
+						EncodePixelFormat = EMediaIOCoreEncodePixelFormat::CharUYVY;
+						OutputFilename = FString::Printf(TEXT("Blackmagic_Output_8_YUV_ch%d"), ChannelInfo.DeviceIndex);
+						break;
+					case BlackmagicDesign::EPixelFormat::pf_10Bits:
+						SampleFormat = EMediaTextureSampleFormat::YUVv210;
+						EncodePixelFormat = EMediaIOCoreEncodePixelFormat::YUVv210;
+						OutputFilename = FString::Printf(TEXT("Blackmagic_Output_10_YUV_ch%d"), ChannelInfo.DeviceIndex);
+						break;
 					}
-					else
+
+					if (bBlackmagicWriteOutputRawDataCmdEnable)
 					{
-						EMediaTextureSampleFormat SampleFormat = EMediaTextureSampleFormat::CharBGRA;
-						EMediaIOCoreEncodePixelFormat EncodePixelFormat = EMediaIOCoreEncodePixelFormat::CharUYVY;
-						FString OutputFilename = "";
+						MediaIOCoreFileWriter::WriteRawFile(OutputFilename, reinterpret_cast<uint8*>(InFrameInfo.VideoBuffer), InFrameInfo.VideoPitch * InFrameInfo.VideoHeight);
+						bBlackmagicWriteOutputRawDataCmdEnable = false;
+					}
 
-						switch (InFrameInfo.PixelFormat)
+					if (bIsProgressivePicture)
+					{
+						if (bEncodeTimecodeInTexel && DecodedTimecode.IsSet())
 						{
-						case BlackmagicDesign::EPixelFormat::pf_8Bits:
-							SampleFormat = EMediaTextureSampleFormat::CharUYVY;
-							EncodePixelFormat = EMediaIOCoreEncodePixelFormat::CharUYVY;
-							OutputFilename = FString::Printf(TEXT("Blackmagic_Output_8_YUV_ch%d"), ChannelInfo.DeviceIndex);
-							break;
-						case BlackmagicDesign::EPixelFormat::pf_10Bits:
-							SampleFormat = EMediaTextureSampleFormat::YUVv210;
-							EncodePixelFormat = EMediaIOCoreEncodePixelFormat::YUVv210;
-							OutputFilename = FString::Printf(TEXT("Blackmagic_Output_10_YUV_ch%d"), ChannelInfo.DeviceIndex);
-							break;
+							FTimecode SetTimecode = DecodedTimecode.GetValue();
+							FMediaIOCoreEncodeTime EncodeTime(EncodePixelFormat, InFrameInfo.VideoBuffer, InFrameInfo.VideoPitch, InFrameInfo.VideoWidth, InFrameInfo.VideoHeight);
+							EncodeTime.Render(SetTimecode.Hours, SetTimecode.Minutes, SetTimecode.Seconds, SetTimecode.Frames);
 						}
 
-						if (bBlackmagicWriteOutputRawDataCmdEnable)
+						auto TextureSample = MediaPlayer->TextureSamplePool->AcquireShared();
+						if (TextureSample->Initialize(InFrameInfo.VideoBuffer
+							, InFrameInfo.VideoPitch * InFrameInfo.VideoHeight
+							, InFrameInfo.VideoPitch
+							, InFrameInfo.VideoWidth
+							, InFrameInfo.VideoHeight
+							, SampleFormat
+							, DecodedTime
+							, MediaPlayer->VideoFrameRate
+							, DecodedTimecode
+							, bIsSRGBInput))
 						{
-							MediaIOCoreFileWriter::WriteRawFile(OutputFilename, reinterpret_cast<uint8*>(InFrameInfo.VideoBuffer), InFrameInfo.VideoPitch * InFrameInfo.VideoHeight);
-							bBlackmagicWriteOutputRawDataCmdEnable = false;
+							MediaPlayer->Samples->AddVideo(TextureSample);
 						}
-
-						if (bIsProgressivePicture)
+					}
+					else
+					{
+						auto TextureSampleEven = MediaPlayer->TextureSamplePool->AcquireShared();
+						if (TextureSampleEven->InitializeWithEvenOddLine(true
+							, InFrameInfo.VideoBuffer
+							, InFrameInfo.VideoPitch * InFrameInfo.VideoHeight
+							, InFrameInfo.VideoPitch
+							, InFrameInfo.VideoWidth
+							, InFrameInfo.VideoHeight
+							, SampleFormat
+							, DecodedTime
+							, MediaPlayer->VideoFrameRate
+							, DecodedTimecode
+							, bIsSRGBInput))
 						{
-							if (bEncodeTimecodeInTexel && DecodedTimecode.IsSet())
-							{
-								FTimecode SetTimecode = DecodedTimecode.GetValue();
-								FMediaIOCoreEncodeTime EncodeTime(EncodePixelFormat, InFrameInfo.VideoBuffer, InFrameInfo.VideoPitch, InFrameInfo.VideoWidth, InFrameInfo.VideoHeight);
-								EncodeTime.Render(SetTimecode.Hours, SetTimecode.Minutes, SetTimecode.Seconds, SetTimecode.Frames);
-							}
-
-							auto TextureSample = MediaPlayer->TextureSamplePool->AcquireShared();
-							if (TextureSample->Initialize(InFrameInfo.VideoBuffer
-								, InFrameInfo.VideoPitch * InFrameInfo.VideoHeight
-								, InFrameInfo.VideoPitch
-								, InFrameInfo.VideoWidth
-								, InFrameInfo.VideoHeight
-								, SampleFormat
-								, DecodedTime
-								, MediaPlayer->VideoFrameRate
-								, DecodedTimecode
-								, bIsSRGBInput))
-							{
-								MediaPlayer->Samples->AddVideo(TextureSample);
-							}
+							MediaPlayer->Samples->AddVideo(TextureSampleEven);
 						}
-						else
+
+						auto TextureSampleOdd = MediaPlayer->TextureSamplePool->AcquireShared();
+						if (TextureSampleOdd->InitializeWithEvenOddLine(false
+							, InFrameInfo.VideoBuffer
+							, InFrameInfo.VideoPitch * InFrameInfo.VideoHeight
+							, InFrameInfo.VideoPitch
+							, InFrameInfo.VideoWidth
+							, InFrameInfo.VideoHeight
+							, SampleFormat
+							, DecodedTimeF2
+							, MediaPlayer->VideoFrameRate
+							, DecodedTimecodeF2
+							, bIsSRGBInput))
 						{
-							auto TextureSampleEven = MediaPlayer->TextureSamplePool->AcquireShared();
-							if (TextureSampleEven->InitializeWithEvenOddLine(true
-								, InFrameInfo.VideoBuffer
-								, InFrameInfo.VideoPitch * InFrameInfo.VideoHeight
-								, InFrameInfo.VideoPitch
-								, InFrameInfo.VideoWidth
-								, InFrameInfo.VideoHeight
-								, SampleFormat
-								, DecodedTime
-								, MediaPlayer->VideoFrameRate
-								, DecodedTimecode
-								, bIsSRGBInput))
-							{
-								MediaPlayer->Samples->AddVideo(TextureSampleEven);
-							}
-
-							auto TextureSampleOdd = MediaPlayer->TextureSamplePool->AcquireShared();
-							if (TextureSampleOdd->InitializeWithEvenOddLine(false
-								, InFrameInfo.VideoBuffer
-								, InFrameInfo.VideoPitch * InFrameInfo.VideoHeight
-								, InFrameInfo.VideoPitch
-								, InFrameInfo.VideoWidth
-								, InFrameInfo.VideoHeight
-								, SampleFormat
-								, DecodedTimeF2
-								, MediaPlayer->VideoFrameRate
-								, DecodedTimecodeF2
-								, bIsSRGBInput))
-							{
-								MediaPlayer->Samples->AddVideo(TextureSampleOdd);
-							}
+							MediaPlayer->Samples->AddVideo(TextureSampleOdd);
 						}
 					}
 				}
@@ -393,9 +371,9 @@ namespace BlackmagicMediaPlayerHelpers
 		uint32 LastNumChannels;
 		uint32 LastSampleRate;
 
-		int32 AudioFrameDropCount;
-		int32 MetadataFrameDropCount;
-		int32 VideoFrameDropCount;
+		/** Frame drop count from the previous tick to keep track of deltas */
+		int32 PreviousAudioFrameDropCount;
+		int32 PreviousVideoFrameDropCount;
 
 		int32 MaxNumAudioFrameBuffer;
 		int32 MaxNumVideoFrameBuffer;
@@ -421,7 +399,9 @@ FBlackmagicMediaPlayer::FBlackmagicMediaPlayer(IMediaEventSink& InEventSink)
 	, AudioSamplePool(new FBlackmagicMediaAudioSamplePool)
 	, TextureSamplePool(new FBlackmagicMediaTextureSamplePool)
 	, bVerifyFrameDropCount(false)
+	, SupportedSampleTypes(static_cast<EMediaIOSampleType>(static_cast<uint8>(EMediaIOSampleType::Video) | static_cast<uint8>(EMediaIOSampleType::Audio)))
 {
+	Samples->EnableTimedDataChannels(SupportedSampleTypes);
 }
 
 FBlackmagicMediaPlayer::~FBlackmagicMediaPlayer()
@@ -511,8 +491,11 @@ bool FBlackmagicMediaPlayer::Open(const FString& Url, const IMediaOptions* Optio
 
 	bVerifyFrameDropCount = Options->GetMediaOption(BlackmagicMediaOption::LogDropFrame, false);
 	const bool bEncodeTimecodeInTexel = TimecodeFormat != EMediaIOTimecodeFormat::None && Options->GetMediaOption(BlackmagicMediaOption::EncodeTimecodeInTexel, false);
-	int32 MaxNumAudioFrameBuffer = Options->GetMediaOption(BlackmagicMediaOption::MaxAudioFrameBuffer, (int64)8);
-	int32 MaxNumVideoFrameBuffer = Options->GetMediaOption(BlackmagicMediaOption::MaxVideoFrameBuffer, (int64)8);
+	MaxNumAudioFrameBuffer = Options->GetMediaOption(BlackmagicMediaOption::MaxAudioFrameBuffer, (int64)8);
+	MaxNumVideoFrameBuffer = Options->GetMediaOption(BlackmagicMediaOption::MaxVideoFrameBuffer, (int64)8);
+
+	// Setup our different supported channels based on source settings
+	SetupSampleChannels();
 
 	bool bSuccess = EventCallback->Initialize(ChannelOptions, bEncodeTimecodeInTexel, MaxNumAudioFrameBuffer, MaxNumVideoFrameBuffer, bIsSRGBInput);
 
@@ -554,8 +537,9 @@ void FBlackmagicMediaPlayer::TickInput(FTimespan DeltaTime, FTimespan Timecode)
 	TickTimeManagement();
 }
 
-void FBlackmagicMediaPlayer::TickFetch(FTimespan DeltaTime, FTimespan /*Timecode*/)
+void FBlackmagicMediaPlayer::TickFetch(FTimespan DeltaTime, FTimespan Timecode)
 {
+	Super::TickFetch(DeltaTime, Timecode);
 	if (IsHardwareReady())
 	{
 		ProcessFrame();
@@ -563,6 +547,13 @@ void FBlackmagicMediaPlayer::TickFetch(FTimespan DeltaTime, FTimespan /*Timecode
 	}
 }
 
+#if WITH_EDITOR
+const FSlateBrush* FBlackmagicMediaPlayer::GetDisplayIcon() const
+{
+	return IBlackmagicMediaModule::Get().GetStyle()->GetBrush("BlackmagicMediaIcon");
+}
+#endif //WITH_EDITOR
+
 void FBlackmagicMediaPlayer::ProcessFrame()
 {
 	EventCallback->UpdateAudioTrackFormat(AudioTrackFormat);
@@ -578,5 +569,16 @@ bool FBlackmagicMediaPlayer::IsHardwareReady() const
 	return EventCallback && EventCallback->GetMediaState() == EMediaState::Playing;
 }
 
+void FBlackmagicMediaPlayer::SetupSampleChannels()
+{
+	FMediaIOSamplingSettings VideoSettings = BaseSettings;
+	VideoSettings.BufferSize = MaxNumVideoFrameBuffer;
+	Samples->InitializeVideoBuffer(VideoSettings);
+
+	FMediaIOSamplingSettings AudioSettings = BaseSettings;
+	AudioSettings.BufferSize = MaxNumAudioFrameBuffer;
+	Samples->InitializeAudioBuffer(AudioSettings);
+}
+
 #undef LOCTEXT_NAMESPACE
 
diff --git a/Engine/Plugins/Media/BlackmagicMedia/Source/BlackmagicMedia/Private/Player/BlackmagicMediaPlayer.h b/Engine/Plugins/Media/BlackmagicMedia/Source/BlackmagicMedia/Private/Player/BlackmagicMediaPlayer.h
--- a/Engine/Plugins/Media/BlackmagicMedia/Source/BlackmagicMedia/Private/Player/BlackmagicMediaPlayer.h
+++ b/Engine/Plugins/Media/BlackmagicMedia/Source/BlackmagicMedia/Private/Player/BlackmagicMediaPlayer.h
@@ -12,6 +12,7 @@
 class IMediaEventSink;
 
 enum class EMediaTextureSampleFormat;
+enum class EMediaIOSampleType;
 
 namespace BlackmagicMediaPlayerHelpers
 {
@@ -66,6 +67,13 @@ public:
 	virtual void TickInput(FTimespan DeltaTime, FTimespan Timecode) override;
 	virtual void TickFetch(FTimespan DeltaTime, FTimespan Timecode) override;
 
+	//~ ITimedDataInput interface
+#if WITH_EDITOR
+	virtual const FSlateBrush* GetDisplayIcon() const override;
+#endif
+
+public:
+
 	/** Process pending audio and video frames, and forward them to the sinks. */
 	void ProcessFrame();
 
@@ -75,6 +83,10 @@ public:
 	/** Is Hardware initialized */
 	virtual bool IsHardwareReady() const override;
 
+protected:
+	/** Setup our different channels with the current set of settings */
+	virtual void SetupSampleChannels() override;
+
 private:
 
 	friend BlackmagicMediaPlayerHelpers::FBlackmagicMediaPlayerEventCallback;
@@ -86,4 +98,11 @@ private:
 
 	/** Log warning about the amount of audio/video frame can't could not be cached . */
 	bool bVerifyFrameDropCount;
+
+	/** Max sample count our different buffer can hold. Taken from MediaSource */
+	int32 MaxNumAudioFrameBuffer = 0;
+	int32 MaxNumVideoFrameBuffer = 0;
+
+	/** Used to flag which sample types we advertise as supported for timed data monitoring */
+	EMediaIOSampleType SupportedSampleTypes;
 };
diff --git a/Engine/Plugins/Media/BlackmagicMedia/Source/BlackmagicMedia/Public/IBlackmagicMediaModule.h b/Engine/Plugins/Media/BlackmagicMedia/Source/BlackmagicMedia/Public/IBlackmagicMediaModule.h
--- a/Engine/Plugins/Media/BlackmagicMedia/Source/BlackmagicMedia/Public/IBlackmagicMediaModule.h
+++ b/Engine/Plugins/Media/BlackmagicMedia/Source/BlackmagicMedia/Public/IBlackmagicMediaModule.h
@@ -3,8 +3,10 @@
 #pragma once
 
 #include "Modules/ModuleInterface.h"
+#include "Modules/ModuleManager.h"
 #include "Templates/SharedPointer.h"
 
+class FSlateStyleSet;
 class IMediaEventSink;
 class IMediaPlayer;
 
@@ -15,6 +17,12 @@ class IBlackmagicMediaModule : public IModuleInterface
 {
 public:
 
+	static inline IBlackmagicMediaModule& Get()
+	{
+		static const FName ModuleName = "BlackmagicMedia";
+		return FModuleManager::LoadModuleChecked<IBlackmagicMediaModule>(ModuleName);
+	}
+
 	/**
 	 * Create an Blackmagic based media player.
 	 * @param EventSink The object that receives media events from the player.
@@ -22,6 +30,9 @@ public:
 	 */
 	virtual TSharedPtr<IMediaPlayer, ESPMode::ThreadSafe> CreatePlayer(IMediaEventSink& EventSink) = 0;
 
+	/** @return SlateStyleSet to be used across the BlackmagicMedia module */
+	virtual TSharedPtr<FSlateStyleSet> GetStyle() = 0;
+
 	/** @return true if the Blackmagic module and VideoIO.dll could be loaded */
 	virtual bool IsInitialized() const = 0;
 
diff --git a/Engine/Plugins/Media/MediaIOFramework/Source/MediaIOCore/Private/Player/MediaIOCorePlayerBase.cpp b/Engine/Plugins/Media/MediaIOFramework/Source/MediaIOCore/Private/Player/MediaIOCorePlayerBase.cpp
--- a/Engine/Plugins/Media/MediaIOFramework/Source/MediaIOCore/Private/Player/MediaIOCorePlayerBase.cpp
+++ b/Engine/Plugins/Media/MediaIOFramework/Source/MediaIOCore/Private/Player/MediaIOCorePlayerBase.cpp
@@ -4,12 +4,15 @@
 
 #include "Engine/Engine.h"
 #include "Engine/TimecodeProvider.h"
+#include "HAL/IConsoleManager.h"
 #include "IMediaEventSink.h"
 #include "IMediaOptions.h"
+#include "ITimeManagementModule.h"
 #include "MediaIOCoreModule.h"
 #include "MediaIOCoreSamples.h"
 #include "Misc/App.h"
 #include "Misc/ScopeLock.h"
+#include "TimedDataInputCollection.h"
 #include "TimeSynchronizableMediaSource.h"
 
 
@@ -32,6 +35,13 @@ namespace MediaIOCorePlayerDetail
 	static const double HighPerformanceClockDelta = ComputeTimeCodeOffset();
 	static bool bLogTimecode = false;
 
+	TAutoConsoleVariable<int32> CVarMediaIOMaxBufferSize(
+		TEXT("MediaIO.TimedDataChannel.MaxBufferSize"),
+		32,
+		TEXT("The max size the MediaIO channels is allowed to set the buffer size."),
+		ECVF_Default
+	);
+
 	static FAutoConsoleCommand CCommandShowTimecode(
 		TEXT("MediaIO.ShowInputTimecode"),
 		TEXT("All media player will log the frame timecode when a new frame is captured."),
@@ -62,13 +72,13 @@ FMediaIOCorePlayerBase::FMediaIOCorePlayerBase(IMediaEventSink& InEventSink)
 	, CurrentTime(FTimespan::Zero())
 	, EventSink(InEventSink)
 	, VideoFrameRate(30, 1)
-	, Samples(new FMediaIOCoreSamples)
 	, bUseTimeSynchronization(false)
 	, bWarnedIncompatibleFrameRate(false)
 	, FrameDelay(0)
 	, TimeDelay(0.0)
 	, PreviousFrameTimespan(FTimespan::Zero())
 {
+	Samples = new FMediaIOCoreSamples(this);
 }
 
 FMediaIOCorePlayerBase::~FMediaIOCorePlayerBase()
@@ -82,14 +92,19 @@ FMediaIOCorePlayerBase::~FMediaIOCorePlayerBase()
 
 void FMediaIOCorePlayerBase::Close()
 {
-	CurrentState = EMediaState::Closed;
-	CurrentTime = FTimespan::Zero();
-	AudioTrackFormat.NumChannels = 0;
-	AudioTrackFormat.SampleRate = 0;
+	if (CurrentState != EMediaState::Closed)
+	{
+		CurrentState = EMediaState::Closed;
+		CurrentTime = FTimespan::Zero();
+		AudioTrackFormat.NumChannels = 0;
+		AudioTrackFormat.SampleRate = 0;
+
+		Samples->FlushSamples();
+		EventSink.ReceiveMediaEvent(EMediaEvent::TracksChanged);
+		EventSink.ReceiveMediaEvent(EMediaEvent::MediaClosed);
 
-	Samples->FlushSamples();
-	EventSink.ReceiveMediaEvent(EMediaEvent::TracksChanged);
-	EventSink.ReceiveMediaEvent(EMediaEvent::MediaClosed);
+		ITimeManagementModule::Get().GetTimedDataInputCollection().Remove(this);
+	}
 }
 
 FString FMediaIOCorePlayerBase::GetInfo() const
@@ -163,6 +178,8 @@ bool FMediaIOCorePlayerBase::Open(const FString& Url, const IMediaOptions* Optio
 {
 	Close();
 
+	ITimeManagementModule::Get().GetTimedDataInputCollection().Add(this);
+
 	OpenUrl = Url;
 	return ReadMediaOptions(Options);
 }
@@ -204,6 +221,11 @@ void FMediaIOCorePlayerBase::TickTimeManagement()
 	}
 }
 
+void FMediaIOCorePlayerBase::TickFetch(FTimespan DeltaTime, FTimespan Timecode)
+{
+	Samples->CacheSamplesState(GetTime());
+}
+
 /* IMediaCache interface
  *****************************************************************************/
 bool FMediaIOCorePlayerBase::QueryCacheState(EMediaCacheState State, TRangeSet<FTimespan>& OutTimeRanges) const
@@ -239,7 +261,6 @@ int32 FMediaIOCorePlayerBase::GetSampleCount(EMediaCacheState State) const
 	return Count;
 }
 
-
 /* IMediaControls interface
  *****************************************************************************/
 
@@ -440,6 +461,13 @@ bool FMediaIOCorePlayerBase::ReadMediaOptions(const IMediaOptions* Options)
 		VideoTrackFormat.FrameRate = VideoFrameRate.AsDecimal();
 		VideoTrackFormat.TypeName = Options->GetMediaOption(FMediaIOCoreMediaOption::VideoModeName, FString(TEXT("1080p30")));
 	}
+
+	//Setup base sampling settings
+	BaseSettings.EvaluationType = bUseTimeSynchronization ? ETimedDataInputEvaluationType::Timecode : ETimedDataInputEvaluationType::PlatformTime;
+	BaseSettings.FrameRate = VideoFrameRate;
+	BaseSettings.PlayerTimeOffset = MediaIOCorePlayerDetail::HighPerformanceClockDelta;
+	BaseSettings.AbsoluteMaxBufferSize = MediaIOCorePlayerDetail::CVarMediaIOMaxBufferSize.GetValueOnGameThread();
+
 	return true;
 }
 
@@ -462,4 +490,96 @@ bool FMediaIOCorePlayerBase::IsTimecodeLogEnabled()
 #endif
 }
 
+TArray<ITimedDataInputChannel*> FMediaIOCorePlayerBase::GetChannels() const
+{
+	return Channels;
+}
+
+FText FMediaIOCorePlayerBase::GetDisplayName() const
+{
+	return FText::Format(LOCTEXT("PlayerDisplayName", "{0} - {1}"), FText::FromName(GetPlayerName()), FText::FromString(GetUrl()));
+}
+
+ETimedDataInputEvaluationType FMediaIOCorePlayerBase::GetEvaluationType() const
+{
+	if (bUseTimeSynchronization)
+	{
+		return ETimedDataInputEvaluationType::Timecode;
+	}
+	else
+	{
+		return ETimedDataInputEvaluationType::PlatformTime;
+	}
+}
+
+void FMediaIOCorePlayerBase::SetEvaluationType(ETimedDataInputEvaluationType Evaluation)
+{
+	bool bNewEvaluation = false;
+	switch (Evaluation)
+	{
+	case ETimedDataInputEvaluationType::Timecode:
+		bNewEvaluation = true;
+		break;
+	case ETimedDataInputEvaluationType::PlatformTime:
+	case ETimedDataInputEvaluationType::None:
+	default:
+		bNewEvaluation = false;
+		break;
+	}
+
+	if (bNewEvaluation != bUseTimeSynchronization)
+	{
+		bUseTimeSynchronization = bNewEvaluation;
+		BaseSettings.EvaluationType = bUseTimeSynchronization ? ETimedDataInputEvaluationType::Timecode : ETimedDataInputEvaluationType::PlatformTime;
+		SetupSampleChannels();
+	}
+}
+
+double FMediaIOCorePlayerBase::GetEvaluationOffsetInSeconds() const
+{
+	if (bUseTimeSynchronization)
+	{
+		return ITimedDataInput::ConvertFrameOffsetInSecondOffset(FrameDelay, VideoFrameRate);
+	}
+	else
+	{
+		return TimeDelay;
+	}
+}
+
+void FMediaIOCorePlayerBase::SetEvaluationOffsetInSeconds(double Offset)
+{
+	if (bUseTimeSynchronization)
+	{
+		//Media doesn't support subframes playback (interpolation between frames) so offsets are always in full frame
+		const double FrameOffset = ITimedDataInput::ConvertSecondOffsetInFrameOffset(Offset, VideoFrameRate);
+		FrameDelay = FMath::FloorToInt(FrameOffset);
+	}
+	else
+	{
+		TimeDelay = Offset;
+	}
+}
+
+FFrameRate FMediaIOCorePlayerBase::GetFrameRate() const
+{
+	return VideoFrameRate;
+}
+
+bool FMediaIOCorePlayerBase::IsDataBufferSizeControlledByInput() const
+{
+	//Each channel (audio, video, etc...) can have a different size
+	return false;
+}
+
+void FMediaIOCorePlayerBase::AddChannel(ITimedDataInputChannel* Channel)
+{
+	Channels.AddUnique(Channel);
+}
+
+void FMediaIOCorePlayerBase::RemoveChannel(ITimedDataInputChannel* Channel)
+{
+	Channels.Remove(Channel);
+}
+
 #undef LOCTEXT_NAMESPACE
diff --git a/Engine/Plugins/Media/MediaIOFramework/Source/MediaIOCore/Private/Player/MediaIOCoreSamples.cpp b/Engine/Plugins/Media/MediaIOFramework/Source/MediaIOCore/Private/Player/MediaIOCoreSamples.cpp
--- a/Engine/Plugins/Media/MediaIOFramework/Source/MediaIOCore/Private/Player/MediaIOCoreSamples.cpp
+++ b/Engine/Plugins/Media/MediaIOFramework/Source/MediaIOCore/Private/Player/MediaIOCoreSamples.cpp
@@ -8,158 +8,98 @@
 #include "IMediaTextureSample.h"
 
 
+FMediaIOCoreSamples::FMediaIOCoreSamples(ITimedDataInput* ChannelsOwner)
+	: VideoSamples(ChannelsOwner, "Video")
+	, AudioSamples(ChannelsOwner, "Audio")
+	, MetadataSamples(ChannelsOwner, "Metadata")
+	, SubtitleSamples(ChannelsOwner, "Subtitles")
+	, CaptionSamples(ChannelsOwner, "Caption")
+{
+
+}
+
+void FMediaIOCoreSamples::CacheSamplesState(FTimespan PlayerTime)
+{
+	VideoSamples.CacheState(PlayerTime);
+	AudioSamples.CacheState(PlayerTime);
+	MetadataSamples.CacheState(PlayerTime);
+	SubtitleSamples.CacheState(PlayerTime);
+	CaptionSamples.CacheState(PlayerTime);
+}
+
+void FMediaIOCoreSamples::EnableTimedDataChannels(EMediaIOSampleType SampleTypes)
+{
+	VideoSamples.EnableChannel(static_cast<uint8>(SampleTypes) & static_cast<uint8>(EMediaIOSampleType::Video));
+	AudioSamples.EnableChannel(static_cast<uint8>(SampleTypes)& static_cast<uint8>(EMediaIOSampleType::Audio));
+	MetadataSamples.EnableChannel(static_cast<uint8>(SampleTypes)& static_cast<uint8>(EMediaIOSampleType::Metadata));
+	SubtitleSamples.EnableChannel(static_cast<uint8>(SampleTypes)& static_cast<uint8>(EMediaIOSampleType::Subtitles));
+	CaptionSamples.EnableChannel(static_cast<uint8>(SampleTypes)& static_cast<uint8>(EMediaIOSampleType::Caption));
+}
+
+void FMediaIOCoreSamples::InitializeVideoBuffer(const FMediaIOSamplingSettings& InSettings)
+{
+	VideoSamples.UpdateSettings(InSettings);
+}
+
+void FMediaIOCoreSamples::InitializeAudioBuffer(const FMediaIOSamplingSettings& InSettings)
+{
+	AudioSamples.UpdateSettings(InSettings);
+}
+
+void FMediaIOCoreSamples::InitializeMetadataBuffer(const FMediaIOSamplingSettings& InSettings)
+{
+	MetadataSamples.UpdateSettings(InSettings);
+}
+
+void FMediaIOCoreSamples::InitializeSubtitlesBuffer(const FMediaIOSamplingSettings& InSettings)
+{
+	SubtitleSamples.UpdateSettings(InSettings);
+}
+
+void FMediaIOCoreSamples::InitializeCaptionBuffer(const FMediaIOSamplingSettings& InSettings)
+{
+	CaptionSamples.UpdateSettings(InSettings);
+}
+
+
 /* IMediaSamples interface
 *****************************************************************************/
 
 bool FMediaIOCoreSamples::FetchAudio(TRange<FTimespan> TimeRange, TSharedPtr<IMediaAudioSample, ESPMode::ThreadSafe>& OutSample)
 {
-	FScopeLock Lock(&AudioCriticalSection);
-
-	const int32 SampleCount = AudioSamples.Num();
-	if (SampleCount > 0)
-	{
-		TSharedPtr<IMediaAudioSample, ESPMode::ThreadSafe> Sample = AudioSamples[SampleCount - 1];
-		if (Sample.IsValid())
-		{
-			const FTimespan SampleTime = Sample->GetTime();
-
-			if (TimeRange.Overlaps(TRange<FTimespan>(SampleTime, SampleTime + Sample->GetDuration())))
-			{
-				AudioSamples.RemoveAt(SampleCount - 1);
-				OutSample = Sample;
-				return true;
-			}
-		}
-	}
-
-	return false;
+	return AudioSamples.FetchSample(TimeRange, OutSample);
 }
 
 
 bool FMediaIOCoreSamples::FetchCaption(TRange<FTimespan> TimeRange, TSharedPtr<IMediaOverlaySample, ESPMode::ThreadSafe>& OutSample)
 {
-	FScopeLock Lock(&CaptionCriticalSection);
-
-	const int32 SampleCount = CaptionSamples.Num();
-	if (SampleCount > 0)
-	{
-		TSharedPtr<IMediaOverlaySample, ESPMode::ThreadSafe> Sample = CaptionSamples[SampleCount - 1];
-		if (Sample.IsValid())
-		{
-			const FTimespan SampleTime = Sample->GetTime();
-
-			if (TimeRange.Overlaps(TRange<FTimespan>(SampleTime, SampleTime + Sample->GetDuration())))
-			{
-				CaptionSamples.RemoveAt(SampleCount - 1);
-				OutSample = Sample;
-				return true;
-			}
-		}
-	}
-
-	return false;
+	return CaptionSamples.FetchSample(TimeRange, OutSample);
 }
 
 
 bool FMediaIOCoreSamples::FetchMetadata(TRange<FTimespan> TimeRange, TSharedPtr<IMediaBinarySample, ESPMode::ThreadSafe>& OutSample)
 {
-	FScopeLock Lock(&MetadataCriticalSection);
-
-	const int32 SampleCount = MetadataSamples.Num();
-	if (SampleCount > 0)
-	{
-		TSharedPtr<IMediaBinarySample, ESPMode::ThreadSafe> Sample = MetadataSamples[SampleCount - 1];
-		if (Sample.IsValid())
-		{
-			const FTimespan SampleTime = Sample->GetTime();
-
-			if (TimeRange.Overlaps(TRange<FTimespan>(SampleTime, SampleTime + Sample->GetDuration())))
-			{
-				MetadataSamples.RemoveAt(SampleCount - 1);
-				OutSample = Sample;
-				return true;
-			}
-		}
-	}
-
-	return false;
+	return MetadataSamples.FetchSample(TimeRange, OutSample);
 }
 
 
 bool FMediaIOCoreSamples::FetchSubtitle(TRange<FTimespan> TimeRange, TSharedPtr<IMediaOverlaySample, ESPMode::ThreadSafe>& OutSample)
 {
-	FScopeLock Lock(&SubtitleCriticalSection);
-
-	const int32 SampleCount = SubtitleSamples.Num();
-	if (SampleCount > 0)
-	{
-		TSharedPtr<IMediaOverlaySample, ESPMode::ThreadSafe> Sample = SubtitleSamples[SampleCount - 1];
-		if (Sample.IsValid())
-		{
-			const FTimespan SampleTime = Sample->GetTime();
-
-			if (TimeRange.Overlaps(TRange<FTimespan>(SampleTime, SampleTime + Sample->GetDuration())))
-			{
-				SubtitleSamples.RemoveAt(SampleCount - 1);
-				OutSample = Sample;
-				return true;
-			}
-		}
-	}
-
-	return false;
+	return SubtitleSamples.FetchSample(TimeRange, OutSample);
 }
 
 
 bool FMediaIOCoreSamples::FetchVideo(TRange<FTimespan> TimeRange, TSharedPtr<IMediaTextureSample, ESPMode::ThreadSafe>& OutSample)
 {
-	FScopeLock Lock(&VideoCriticalSection);
-
-	const int32 SampleCount = VideoSamples.Num();
-	if (SampleCount > 0)
-	{
-		TSharedPtr<IMediaTextureSample, ESPMode::ThreadSafe> Sample = VideoSamples[SampleCount - 1];
-		if (Sample.IsValid())
-		{
-			const FTimespan SampleTime = Sample->GetTime();
-
-			if (TimeRange.Overlaps(TRange<FTimespan>(SampleTime, SampleTime + Sample->GetDuration())))
-			{
-				VideoSamples.RemoveAt(SampleCount - 1);
-				OutSample = Sample;
-				return true;
-			}
-		}
-	}
-
-	return false;
+	return VideoSamples.FetchSample(TimeRange, OutSample);
 }
 
 
 void FMediaIOCoreSamples::FlushSamples()
 {
-	{
-		FScopeLock Lock(&AudioCriticalSection); 
-		AudioSamples.Empty();
-	}
-
-	{
-		FScopeLock Lock(&CaptionCriticalSection);
-		CaptionSamples.Empty();
-	}
-
-	{
-		FScopeLock Lock(&MetadataCriticalSection);
-		MetadataSamples.Empty();
-	}
-
-	{
-		FScopeLock Lock(&SubtitleCriticalSection);
-		SubtitleSamples.Empty();
-	}
-
-	{
-		FScopeLock Lock(&VideoCriticalSection);
-		VideoSamples.Empty();
-	}
+	AudioSamples.FlushSamples();
+	CaptionSamples.FlushSamples();
+	MetadataSamples.FlushSamples();
+	SubtitleSamples.FlushSamples();
+	VideoSamples.FlushSamples();
 }
diff --git a/Engine/Plugins/Media/MediaIOFramework/Source/MediaIOCore/Public/MediaIOCorePlayerBase.h b/Engine/Plugins/Media/MediaIOFramework/Source/MediaIOCore/Public/MediaIOCorePlayerBase.h
--- a/Engine/Plugins/Media/MediaIOFramework/Source/MediaIOCore/Public/MediaIOCorePlayerBase.h
+++ b/Engine/Plugins/Media/MediaIOFramework/Source/MediaIOCore/Public/MediaIOCorePlayerBase.h
@@ -9,8 +9,10 @@
 #include "IMediaPlayer.h"
 #include "IMediaTracks.h"
 #include "IMediaView.h"
+#include "ITimedDataInput.h"
 
 #include "HAL/CriticalSection.h"
+#include "MediaIOCoreSampleContainer.h"
 #include "Misc/CoreMisc.h"
 
 #include "Misc/FrameRate.h"
@@ -29,7 +31,6 @@ struct MEDIAIOCORE_API FMediaIOCoreMediaOption
 	static const FName VideoModeName;
 };
 
-
 /**
  * Implements a base player for hardware IO cards. 
  *
@@ -47,6 +48,7 @@ class MEDIAIOCORE_API FMediaIOCorePlayerBase
 	, protected IMediaControls
 	, protected IMediaTracks
 	, protected IMediaView
+	, public ITimedDataInput
 {
 public:
 
@@ -78,6 +80,7 @@ public:
 	virtual FString GetUrl() const override;
 	virtual IMediaView& GetView() override;
 	virtual void TickTimeManagement();
+	virtual void TickFetch(FTimespan DeltaTime, FTimespan Timecode) override;
 
 public:
 	//~ IMediaCache interface
@@ -117,6 +120,19 @@ protected:
 	virtual bool SelectTrack(EMediaTrackType TrackType, int32 TrackIndex) override;
 	virtual bool SetTrackFormat(EMediaTrackType TrackType, int32 TrackIndex, int32 FormatIndex) override;
 
+public:
+	//~ ITimedDataInput interface
+	virtual FText GetDisplayName() const override;
+	virtual TArray<ITimedDataInputChannel*> GetChannels() const override;
+	virtual ETimedDataInputEvaluationType GetEvaluationType() const override;
+	virtual void SetEvaluationType(ETimedDataInputEvaluationType Evaluation) override;
+	virtual double GetEvaluationOffsetInSeconds() const override;
+	virtual void SetEvaluationOffsetInSeconds(double Offset) override;
+	virtual FFrameRate GetFrameRate() const override;
+	virtual bool IsDataBufferSizeControlledByInput() const override;
+	virtual void AddChannel(ITimedDataInputChannel* Channel) override;
+	virtual void RemoveChannel(ITimedDataInputChannel * Channel) override; 
+
 protected:
 	/** Is the IO hardware/device ready to be used. */
 	virtual bool IsHardwareReady() const = 0;
@@ -132,6 +148,12 @@ protected:
 
 	/** Log the timecode when a frame is received. */
 	static bool IsTimecodeLogEnabled();
+	
+	/** 
+	 * Setup settings for the different kind of supported data channels. 
+	 * PlayerBase will setup the common settings
+	 */
+	virtual void SetupSampleChannels() = 0;
 
 protected:
 	/** Critical section for synchronizing access to receiver and sinks. */
@@ -175,6 +197,12 @@ protected:
 
 	/** Previous frame Timespan */
 	FTimespan PreviousFrameTimespan;
+
+	/** Timed Data Input handler */
+	TArray<ITimedDataInputChannel*> Channels;
+	
+	/** Base set of settings to start from when setuping channels */
+	FMediaIOSamplingSettings BaseSettings;
 };
 
 
diff --git a/Engine/Plugins/Media/MediaIOFramework/Source/MediaIOCore/Public/MediaIOCoreSamples.h b/Engine/Plugins/Media/MediaIOFramework/Source/MediaIOCore/Public/MediaIOCoreSamples.h
--- a/Engine/Plugins/Media/MediaIOFramework/Source/MediaIOCore/Public/MediaIOCoreSamples.h
+++ b/Engine/Plugins/Media/MediaIOFramework/Source/MediaIOCore/Public/MediaIOCoreSamples.h
@@ -5,6 +5,7 @@
 #include "CoreMinimal.h"
 #include "IMediaSamples.h"
 #include "IMediaTextureSample.h"
+#include "MediaIOCoreSampleContainer.h"
 #include "Misc/ScopeLock.h"
 #include "Templates/SharedPointer.h"
 
@@ -13,14 +14,28 @@ class IMediaBinarySample;
 class IMediaOverlaySample;
 class IMediaTextureSample;
 
+/**
+ * Bitflags about different types of sample a MediaIO can push
+ */
+UENUM()
+enum class EMediaIOSampleType
+{
+	None = 0,
+	Video = 1<<0,
+	Audio = 1<<1,
+	Metadata = 1<<2,
+	Subtitles = 1<<3,
+	Caption = 1<<4,
+};
+
 /**
  * General purpose media sample queue.
  */
-class MEDIAIOCORE_API FMediaIOCoreSamples
-	: public IMediaSamples
+
+class MEDIAIOCORE_API FMediaIOCoreSamples : public IMediaSamples
 {
 public:
-	FMediaIOCoreSamples() = default;
+	FMediaIOCoreSamples(ITimedDataInput* ChannelsOwner);
 	FMediaIOCoreSamples(const FMediaIOCoreSamples&) = delete;
 	FMediaIOCoreSamples& operator=(const FMediaIOCoreSamples&) = delete;
 
@@ -35,9 +50,7 @@ public:
 	 */
 	bool AddAudio(const TSharedRef<IMediaAudioSample, ESPMode::ThreadSafe>& Sample)
 	{
-		FScopeLock Lock(&AudioCriticalSection);
-		AudioSamples.EmplaceAt(0, Sample);
-		return true;
+		return AudioSamples.AddSample(Sample);
 	}
 
 	/**
@@ -49,9 +62,7 @@ public:
 	 */
 	bool AddCaption(const TSharedRef<IMediaOverlaySample, ESPMode::ThreadSafe>& Sample)
 	{
-		FScopeLock Lock(&CaptionCriticalSection);
-		CaptionSamples.EmplaceAt(0, Sample);
-		return true;
+		return CaptionSamples.AddSample(Sample);
 	}
 
 	/**
@@ -63,9 +74,7 @@ public:
 	 */
 	bool AddMetadata(const TSharedRef<IMediaBinarySample, ESPMode::ThreadSafe>& Sample)
 	{
-		FScopeLock Lock(&MetadataCriticalSection);
-		MetadataSamples.EmplaceAt(0, Sample);
-		return true;
+		return MetadataSamples.AddSample(Sample);
 	}
 
 	/**
@@ -77,9 +86,7 @@ public:
 	 */
 	bool AddSubtitle(const TSharedRef<IMediaOverlaySample, ESPMode::ThreadSafe>& Sample)
 	{
-		FScopeLock Lock(&SubtitleCriticalSection);
-		SubtitleSamples.EmplaceAt(0, Sample);
-		return true;
+		return SubtitleSamples.AddSample(Sample);
 	}
 
 	/**
@@ -91,9 +98,7 @@ public:
 	 */
 	bool AddVideo(const TSharedRef<IMediaTextureSample, ESPMode::ThreadSafe>& Sample)
 	{
-		FScopeLock Lock(&VideoCriticalSection);
-		VideoSamples.EmplaceAt(0, Sample);
-		return true;
+		return VideoSamples.AddSample(Sample);
 	}
 
 	/**
@@ -104,14 +109,7 @@ public:
 	 */
 	bool PopAudio()
 	{
-		FScopeLock Lock(&AudioCriticalSection);
-		const int32 SampleCount = AudioSamples.Num();
-		if (SampleCount > 0)
-		{
-			AudioSamples.RemoveAt(SampleCount - 1);
-			return true;
-		}
-		return false;
+		return AudioSamples.PopSample();
 	}
 
 	/**
@@ -122,14 +120,7 @@ public:
 	 */
 	bool PopCaption()
 	{
-		FScopeLock Lock(&CaptionCriticalSection);
-		const int32 SampleCount = CaptionSamples.Num();
-		if (SampleCount > 0)
-		{
-			CaptionSamples.RemoveAt(SampleCount - 1);
-			return true;
-		}
-		return false;
+		return CaptionSamples.PopSample();
 	}
 
 	/**
@@ -140,14 +131,7 @@ public:
 	 */
 	bool PopMetadata()
 	{
-		FScopeLock Lock(&MetadataCriticalSection);
-		const int32 SampleCount = MetadataSamples.Num();
-		if (SampleCount > 0)
-		{
-			MetadataSamples.RemoveAt(SampleCount - 1);
-			return true;
-		}
-		return false;
+		return MetadataSamples.PopSample();
 	}
 
 	/**
@@ -158,14 +142,7 @@ public:
 	 */
 	bool PopSubtitle()
 	{
-		FScopeLock Lock(&SubtitleCriticalSection);
-		const int32 SampleCount = SubtitleSamples.Num();
-		if (SampleCount > 0)
-		{
-			SubtitleSamples.RemoveAt(SampleCount - 1);
-			return true;
-		}
-		return false;
+		return SubtitleSamples.PopSample();
 	}
 
 	/**
@@ -176,14 +153,7 @@ public:
 	 */
 	bool PopVideo()
 	{
-		FScopeLock Lock(&VideoCriticalSection);
-		const int32 SampleCount = VideoSamples.Num();
-		if (SampleCount > 0)
-		{
-			VideoSamples.RemoveAt(SampleCount - 1);
-			return true;
-		}
-		return false;
+		return VideoSamples.PopSample();
 	}
 
 	/**
@@ -194,8 +164,7 @@ public:
 	 */
 	int32 NumAudioSamples() const
 	{
-		FScopeLock Lock(&AudioCriticalSection);
-		return AudioSamples.Num();
+		return AudioSamples.NumSamples();
 	}
 
 	/**
@@ -206,8 +175,7 @@ public:
 	 */
 	int32 NumCaptionSamples() const
 	{
-		FScopeLock Lock(&CaptionCriticalSection);
-		return CaptionSamples.Num();
+		return CaptionSamples.NumSamples();
 	}
 
 	/**
@@ -218,8 +186,7 @@ public:
 	 */
 	int32 NumMetadataSamples() const
 	{
-		FScopeLock Lock(&MetadataCriticalSection);
-		return MetadataSamples.Num();
+		return MetadataSamples.NumSamples();
 	}
 
 	/**
@@ -230,8 +197,7 @@ public:
 	 */
 	int32 NumSubtitleSamples() const
 	{
-		FScopeLock Lock(&SubtitleCriticalSection);
-		return SubtitleSamples.Num();
+		return SubtitleSamples.NumSamples();
 	}
 
 	/**
@@ -242,8 +208,7 @@ public:
 	 */
 	int32 NumVideoSamples() const
 	{
-		FScopeLock Lock(&VideoCriticalSection);
-		return VideoSamples.Num();
+		return VideoSamples.NumSamples();
 	}
 
 	/**
@@ -254,21 +219,74 @@ public:
 	 */
 	FTimespan GetNextVideoSampleTime()
 	{
-		FScopeLock Lock(&VideoCriticalSection);
-
-		const int32 SampleCount = VideoSamples.Num();
-		if (SampleCount > 0)
-		{
-			TSharedPtr<IMediaTextureSample, ESPMode::ThreadSafe> Sample = VideoSamples[SampleCount - 1];
-			if (Sample.IsValid())
-			{
-				return Sample->GetTime();
-			}
-		}
-
-		return FTimespan();
+		return VideoSamples.GetNextSampleTime();
+	}
+
+	/**
+	 * Get Audio Samples frame dropped count.
+	 *
+	 * @return Number of frames dropped
+	 */
+	int32 GetAudioFrameDropCount() const
+	{
+		return AudioSamples.GetFrameDroppedStat();
+	}
+
+	/**
+	 * Get Video Samples frame dropped count.
+	 *
+	 * @return Number of frames dropped
+	 */
+	int32 GetVideoFrameDropCount() const
+	{
+		return VideoSamples.GetFrameDroppedStat();
+	}
+	
+	/**
+	 * Get Metadata Samples frame dropped count.
+	 *
+	 * @return Number of frames dropped
+	 */
+	int32 GetMetadataFrameDropCount() const
+	{
+		return MetadataSamples.GetFrameDroppedStat();
+	}
+
+	/**
+	 * Get Subtitles Samples frame dropped count.
+	 *
+	 * @return Number of frames dropped
+	 */
+	int32 GetSubtitlesFrameDropCount() const
+	{
+		return SubtitleSamples.GetFrameDroppedStat();
+	}
+
+	/**
+	 * Get Caption Samples frame dropped count.
+	 *
+	 * @return Number of frames dropped
+	 */
+	int32 GetCaptionsFrameDropCount() const
+	{
+		return CaptionSamples.GetFrameDroppedStat();
 	}
 
+	/**
+	 * Caches the current sample container state for a given Player (evaluation) time
+	 */
+	void CacheSamplesState(FTimespan PlayerTime);
+
+	/** Enable or disable channels based on bitfield flag */
+	void EnableTimedDataChannels(EMediaIOSampleType SampleTypes);
+
+	/** Initialize our different buffers with player's settings and wheter or not it should be displayed / supported in Timing monitor */
+	void InitializeVideoBuffer(const FMediaIOSamplingSettings& InSettings);
+	void InitializeAudioBuffer(const FMediaIOSamplingSettings& InSettings);
+	void InitializeMetadataBuffer(const FMediaIOSamplingSettings& InSettings);
+	void InitializeSubtitlesBuffer(const FMediaIOSamplingSettings& InSettings);
+	void InitializeCaptionBuffer(const FMediaIOSamplingSettings& InSettings);
+
 public:
 
 	//~ IMediaSamples interface
@@ -282,23 +300,9 @@ public:
 
 protected:
 
-	/** Audio sample queue. */
-	mutable FCriticalSection AudioCriticalSection;
-	TArray<TSharedPtr<IMediaAudioSample, ESPMode::ThreadSafe>> AudioSamples;
-
-	/** Caption sample queue. */
-	mutable FCriticalSection CaptionCriticalSection;
-	TArray<TSharedPtr<IMediaOverlaySample, ESPMode::ThreadSafe>> CaptionSamples;
-
-	/** Metadata sample queue. */
-	mutable FCriticalSection MetadataCriticalSection;
-	TArray<TSharedPtr<IMediaBinarySample, ESPMode::ThreadSafe>> MetadataSamples;
-
-	/** Subtitle sample queue. */
-	mutable FCriticalSection SubtitleCriticalSection;
-	TArray<TSharedPtr<IMediaOverlaySample, ESPMode::ThreadSafe>> SubtitleSamples;
-
-	/** Video sample queue. */
-	mutable FCriticalSection VideoCriticalSection;
-	TArray<TSharedPtr<IMediaTextureSample, ESPMode::ThreadSafe>> VideoSamples;
+	FMediaIOCoreSampleContainer<IMediaTextureSample> VideoSamples;
+	FMediaIOCoreSampleContainer<IMediaAudioSample> AudioSamples;
+	FMediaIOCoreSampleContainer<IMediaBinarySample> MetadataSamples;
+	FMediaIOCoreSampleContainer<IMediaOverlaySample> SubtitleSamples;
+	FMediaIOCoreSampleContainer<IMediaOverlaySample> CaptionSamples;
 };
