- Capped Aja Timecode provider frame rate to 30 when read from reference pin until it's known to work for higher fps

diff --git a/Engine/Plugins/Media/AjaMedia/Source/AjaMedia/Private/Aja/AjaDeviceProvider.cpp b/Engine/Plugins/Media/AjaMedia/Source/AjaMedia/Private/Aja/AjaDeviceProvider.cpp
--- a/Engine/Plugins/Media/AjaMedia/Source/AjaMedia/Private/Aja/AjaDeviceProvider.cpp
+++ b/Engine/Plugins/Media/AjaMedia/Source/AjaMedia/Private/Aja/AjaDeviceProvider.cpp
@@ -805,8 +805,11 @@ TArray<FAjaMediaTimecodeReference> FAjaDeviceProvider::GetTimecodeReferences() c
 
 				for (const FCommonFrameRateInfo& FrameRate : AllFrameRates)
 				{
-					DefaultFAjaMediaTimecodeReference.LtcFrameRate = FrameRate.FrameRate;
-					Results.Add(DefaultFAjaMediaTimecodeReference);
+					if (FrameRate.FrameRate.AsDecimal() <= (30.0f + SMALL_NUMBER))
+					{
+						DefaultFAjaMediaTimecodeReference.LtcFrameRate = FrameRate.FrameRate;
+						Results.Add(DefaultFAjaMediaTimecodeReference);
+					}
 				}
 			}
 		}
