Fixing an EditorFramework dependency that was added to the wrong scope in CL 14114839

#rnx

diff --git a/Engine/Plugins/Media/AjaMedia/Source/AjaMediaOutput/AjaMediaOutput.Build.cs b/Engine/Plugins/Media/AjaMedia/Source/AjaMediaOutput/AjaMediaOutput.Build.cs
--- a/Engine/Plugins/Media/AjaMedia/Source/AjaMediaOutput/AjaMediaOutput.Build.cs
+++ b/Engine/Plugins/Media/AjaMedia/Source/AjaMediaOutput/AjaMediaOutput.Build.cs
@@ -25,7 +25,6 @@ namespace UnrealBuildTool.Rules
 					"AJA",
 					"Core",
 					"CoreUObject",
-					"EditorFramework",
 					"Engine",
 					"MovieSceneCapture",
 					"RenderCore",
@@ -38,6 +37,7 @@ namespace UnrealBuildTool.Rules
 
 			if (Target.bBuildEditor == true)
 			{
+				PrivateDependencyModuleNames.Add("EditorFramework");
 				PrivateDependencyModuleNames.Add("UnrealEd");
 			}
 		}
