AJA: Added support for Pause.

diff --git a/Engine/Plugins/Media/AjaMedia/Source/AjaMedia/Private/Player/AjaMediaPlayer.cpp b/Engine/Plugins/Media/AjaMedia/Source/AjaMedia/Private/Player/AjaMediaPlayer.cpp
--- a/Engine/Plugins/Media/AjaMedia/Source/AjaMedia/Private/Player/AjaMediaPlayer.cpp
+++ b/Engine/Plugins/Media/AjaMedia/Source/AjaMedia/Private/Player/AjaMediaPlayer.cpp
@@ -69,6 +69,7 @@ FAjaMediaPlayer::FAjaMediaPlayer(IMediaEventSink& InEventSink)
 	, bVerifyFrameDropCount(true)
 	, InputChannel(nullptr)
 	, SupportedSampleTypes(EMediaIOSampleType::None)
+	, bPauseRequested(false)
 {
 }
 
@@ -496,11 +497,13 @@ bool FAjaMediaPlayer::OnInputFrameReceived(const AJA::AJAInputFrameData& InInput
 {
 	SCOPE_CYCLE_COUNTER(STAT_AJA_MediaPlayer_ProcessFrame);
 
-	if (AjaThreadNewState != EMediaState::Playing)
+	if ((AjaThreadNewState != EMediaState::Playing) && (AjaThreadNewState != EMediaState::Paused))
 	{
 		return false;
 	}
 
+	AjaThreadNewState = bPauseRequested ? EMediaState::Paused : EMediaState::Playing;
+
 	AjaThreadFrameDropCount = InInputFrame.FramesDropped;
 
 	FTimespan DecodedTime = FTimespan::FromSeconds(GetPlatformSeconds());
@@ -700,7 +703,7 @@ bool FAjaMediaPlayer::OnOutputFrameCopied(const AJA::AJAOutputFrameData& InFrame
 
 bool FAjaMediaPlayer::IsHardwareReady() const
 {
-	return AjaThreadNewState == EMediaState::Playing ? true : false;
+	return (AjaThreadNewState == EMediaState::Playing) || (AjaThreadNewState == EMediaState::Paused);
 }
 
 void FAjaMediaPlayer::SetupSampleChannels()
@@ -718,6 +721,23 @@ void FAjaMediaPlayer::SetupSampleChannels()
 	Samples->InitializeMetadataBuffer(MetadataSettings);
 }
 
+bool FAjaMediaPlayer::SetRate(float Rate)
+{
+	if (FMath::IsNearlyEqual(Rate, 1.0f))
+	{
+		bPauseRequested = false;
+		return true;
+	}
+
+	if (FMath::IsNearlyEqual(Rate, 0.0f))
+	{
+		bPauseRequested = true;
+		return true;
+	}
+
+	return false;
+}
+
 #undef LOCTEXT_NAMESPACE
 
 #include "AjaMediaHidePlatformTypes.h"
diff --git a/Engine/Plugins/Media/AjaMedia/Source/AjaMedia/Private/Player/AjaMediaPlayer.h b/Engine/Plugins/Media/AjaMedia/Source/AjaMedia/Private/Player/AjaMediaPlayer.h
--- a/Engine/Plugins/Media/AjaMedia/Source/AjaMedia/Private/Player/AjaMediaPlayer.h
+++ b/Engine/Plugins/Media/AjaMedia/Source/AjaMedia/Private/Player/AjaMediaPlayer.h
@@ -7,6 +7,8 @@
 #include "AjaMediaPrivate.h"
 #include "AjaMediaSource.h"
 
+#include <atomic>
+
 class FAjaMediaAudioSample;
 class FAjaMediaAudioSamplePool;
 class FAjaMediaBinarySamplePool;
@@ -66,6 +68,8 @@ public:
 
 	virtual FString GetStats() const override;
 
+	bool SetRate(float Rate) override;
+
 	//~ ITimedDataInput interface
 #if WITH_EDITOR
 	virtual const FSlateBrush* GetDisplayIcon() const override;
@@ -161,6 +165,10 @@ private:
 
 	/** Frame Description from capture device */
 	AJA::FAJAVideoFormat LastVideoFormatIndex;
+
 	/** Previous frame timecode for stats purpose */
 	AJA::FTimecode AjaThreadPreviousFrameTimecode;
+
+	/** Flag to indicate that pause is being requested */
+	std::atomic<bool> bPauseRequested;
 };
