- Fixing Aja kona 5 4k support with retail firmware

diff --git a/Engine/Plugins/Media/AjaMedia/Source/AjaMedia/Private/Aja/AjaDeviceProvider.cpp b/Engine/Plugins/Media/AjaMedia/Source/AjaMedia/Private/Aja/AjaDeviceProvider.cpp
--- a/Engine/Plugins/Media/AjaMedia/Source/AjaMedia/Private/Aja/AjaDeviceProvider.cpp
+++ b/Engine/Plugins/Media/AjaMedia/Source/AjaMedia/Private/Aja/AjaDeviceProvider.cpp
@@ -482,9 +482,9 @@ TArray<FMediaIOConfiguration> FAjaDeviceProvider::GetConfigurations(bool bAllowI
 										MediaConfiguration.MediaConnection.TransportType = EMediaIOTransportType::SingleLink;
 										// in single TSI, we only support input 1, 3, 5, 7
 										// otherwise it will be single-quad
-										const double MaximalFrameRateForUHD_6G = 30.01;
-										const int32 Increment = MediaConfiguration.MediaMode.FrameRate.AsDecimal() > MaximalFrameRateForUHD_6G ? 4 : 2;
-										for (int32 SourceIndex = 0; SourceIndex < SdiPortCount/ Increment; ++SourceIndex)
+										constexpr int32 Increment = 2;
+										const int32 MaxSourceIndex = (MediaConfiguration.bIsInput && Descriptor.bIs4K) ? 1 : SdiPortCount / Increment; //Corner case for Kona 5 retail firmware. Can do 4k input only on SDI 1.
+										for (int32 SourceIndex = 0; SourceIndex < MaxSourceIndex; ++SourceIndex)
 										{
 											MediaConfiguration.MediaConnection.PortIdentifier = (SourceIndex*Increment) + 1;
 											Results.Add(MediaConfiguration);
