Aja: Fix aja and blackmagic crash on unsupported hardware

diff --git a/Engine/Plugins/Media/AjaMedia/Source/AjaMediaOutput/Private/AjaMediaOutputModule.cpp b/Engine/Plugins/Media/AjaMedia/Source/AjaMediaOutput/Private/AjaMediaOutputModule.cpp
--- a/Engine/Plugins/Media/AjaMedia/Source/AjaMediaOutput/Private/AjaMediaOutputModule.cpp
+++ b/Engine/Plugins/Media/AjaMedia/Source/AjaMediaOutput/Private/AjaMediaOutputModule.cpp
@@ -5,6 +5,7 @@
 #include "AjaMediaFrameGrabberProtocol.h"
 #include "AjaLib.h"
 #include "DynamicRHI.h"
+#include "GenericPlatform/GenericPlatformDriver.h"
 #include "RenderingThread.h"
 #include "RHI.h"
 
@@ -14,51 +15,57 @@ DEFINE_LOG_CATEGORY(LogAjaMediaOutput);
 
 void FAjaMediaOutputModule::StartupModule()
 {
-	ENQUEUE_RENDER_COMMAND(AjaMediaCaptureInitialize)(
-		[this](FRHICommandListImmediate& RHICmdList) mutable
-		{
-			if (!GDynamicRHI)
-			{
-				return;
-			}
+	const FGPUDriverInfo GPUDriverInfo = FPlatformMisc::GetGPUDriverInfo(GRHIAdapterName);
+	bIsGPUTextureTransferAvailable = GPUDriverInfo.IsNVIDIA();
 
-			auto GetRHI = []()
+	if (bIsGPUTextureTransferAvailable)
+	{
+		ENQUEUE_RENDER_COMMAND(AjaMediaCaptureInitialize)(
+			[this](FRHICommandListImmediate& RHICmdList) mutable
 			{
-				FString RHIName = GDynamicRHI->GetName();
-				if (RHIName == TEXT("D3D11"))
+				if (!GDynamicRHI)
 				{
-					return AJA::ERHI::D3D11;
+					return;
 				}
-				else if (RHIName == TEXT("D3D12"))
-				{
-					return AJA::ERHI::D3D12;
-				}
-				else if (RHIName == TEXT("Vulkan"))
+
+				auto GetRHI = []()
 				{
-					return AJA::ERHI::Vulkan;
-				}
+					FString RHIName = GDynamicRHI->GetName();
+					if (RHIName == TEXT("D3D11"))
+					{
+						return AJA::ERHI::D3D11;
+					}
+					else if (RHIName == TEXT("D3D12"))
+					{
+						return AJA::ERHI::D3D12;
+					}
+					else if (RHIName == TEXT("Vulkan"))
+					{
+						return AJA::ERHI::Vulkan;
+					}
 
-				return AJA::ERHI::Invalid;
+					return AJA::ERHI::Invalid;
 
-			};
+				};
 
-			AJA::FInitializeDMAArgs Args;
-			AJA::ERHI RHI = GetRHI();
-			Args.RHI = RHI;
-			/* Re-enable when adding vulkan support
-			if (RHI == AJA::ERHI::Vulkan)
-			{
-				FVulkanDynamicRHI* vkDynamicRHI = static_cast<FVulkanDynamicRHI*>(GDynamicRHI);
-				Args.InVulkanInstance = vkDynamicRHI->GetInstance();
+				AJA::FInitializeDMAArgs Args;
+				AJA::ERHI RHI = GetRHI();
+				Args.RHI = RHI;
+				/* Re-enable when adding vulkan support
+				if (RHI == AJA::ERHI::Vulkan)
+				{
+					FVulkanDynamicRHI* vkDynamicRHI = static_cast<FVulkanDynamicRHI*>(GDynamicRHI);
+					Args.InVulkanInstance = vkDynamicRHI->GetInstance();
 
-				FMemory::Memcpy(Args.InRHIDeviceUUID, vkDynamicRHI->GetDevice()->GetDeviceIdProperties().deviceUUID, 16);
-			}
-			*/
-			Args.RHIDevice = GDynamicRHI->RHIGetNativeDevice();
-			Args.RHICommandQueue = GDynamicRHI->RHIGetNativeGraphicsQueue();
+					FMemory::Memcpy(Args.InRHIDeviceUUID, vkDynamicRHI->GetDevice()->GetDeviceIdProperties().deviceUUID, 16);
+				}
+				*/
+				Args.RHIDevice = GDynamicRHI->RHIGetNativeDevice();
+				Args.RHICommandQueue = GDynamicRHI->RHIGetNativeGraphicsQueue();
 
-			bIsGPUTextureTransferAvailable = AJA::InitializeDMA(Args);
-		});
+				bIsGPUTextureTransferAvailable = AJA::InitializeDMA(Args);
+			});
+	}
 }
 
 void FAjaMediaOutputModule::ShutdownModule()
diff --git a/Engine/Plugins/Media/BlackmagicMedia/Source/BlackmagicMediaOutput/Private/BlackmagicMediaOutputModule.cpp b/Engine/Plugins/Media/BlackmagicMedia/Source/BlackmagicMediaOutput/Private/BlackmagicMediaOutputModule.cpp
--- a/Engine/Plugins/Media/BlackmagicMedia/Source/BlackmagicMediaOutput/Private/BlackmagicMediaOutputModule.cpp
+++ b/Engine/Plugins/Media/BlackmagicMedia/Source/BlackmagicMediaOutput/Private/BlackmagicMediaOutputModule.cpp
@@ -3,6 +3,7 @@
 #include "BlackmagicMediaOutputModule.h"
 
 #include "BlackmagicLib.h"
+#include "GenericPlatform/GenericPlatformDriver.h"
 #include "RenderingThread.h"
 #include "RHI.h"
 
@@ -15,42 +16,48 @@ FBlackmagicMediaOutputModule& FBlackmagicMediaOutputModule::Get()
 
 void FBlackmagicMediaOutputModule::StartupModule()
 {
-	ENQUEUE_RENDER_COMMAND(BlackmagicMediaCaptureInitialize)(
-		[this](FRHICommandListImmediate& RHICmdList) mutable
-		{
-			if (!GDynamicRHI)
-			{
-				return;
-			}
+	const FGPUDriverInfo GPUDriverInfo = FPlatformMisc::GetGPUDriverInfo(GRHIAdapterName);
+	bIsGPUTextureTransferAvailable = GPUDriverInfo.IsNVIDIA();
 
-			auto GetRHI = []()
+	if (bIsGPUTextureTransferAvailable)
+	{
+		ENQUEUE_RENDER_COMMAND(BlackmagicMediaCaptureInitialize)(
+			[this](FRHICommandListImmediate& RHICmdList) mutable
 			{
-				FString RHIName = GDynamicRHI->GetName();
-				if (RHIName == TEXT("D3D11"))
-				{
-					return BlackmagicDesign::ERHI::D3D11;
-				}
-				else if (RHIName == TEXT("D3D12"))
+				if (!GDynamicRHI)
 				{
-					return BlackmagicDesign::ERHI::D3D12;
+					return;
 				}
-				else if (RHIName == TEXT("Vulkan"))
+
+				auto GetRHI = []()
 				{
-					return BlackmagicDesign::ERHI::Vulkan;
-				}
+					FString RHIName = GDynamicRHI->GetName();
+					if (RHIName == TEXT("D3D11"))
+					{
+						return BlackmagicDesign::ERHI::D3D11;
+					}
+					else if (RHIName == TEXT("D3D12"))
+					{
+						return BlackmagicDesign::ERHI::D3D12;
+					}
+					else if (RHIName == TEXT("Vulkan"))
+					{
+						return BlackmagicDesign::ERHI::Vulkan;
+					}
 
-				return  BlackmagicDesign::ERHI::Invalid;
-			};
+					return  BlackmagicDesign::ERHI::Invalid;
+				};
 
 
-			BlackmagicDesign::FInitializeDMAArgs Args;
-			BlackmagicDesign::ERHI RHI = GetRHI();
-			Args.RHI = RHI;
-			Args.RHIDevice = GDynamicRHI->RHIGetNativeDevice();
-			Args.RHICommandQueue = GDynamicRHI->RHIGetNativeGraphicsQueue();
+				BlackmagicDesign::FInitializeDMAArgs Args;
+				BlackmagicDesign::ERHI RHI = GetRHI();
+				Args.RHI = RHI;
+				Args.RHIDevice = GDynamicRHI->RHIGetNativeDevice();
+				Args.RHICommandQueue = GDynamicRHI->RHIGetNativeGraphicsQueue();
 
-			bIsGPUTextureTransferAvailable = BlackmagicDesign::InitializeDMA(Args);
-		});
+				bIsGPUTextureTransferAvailable = BlackmagicDesign::InitializeDMA(Args);
+			});
+	}
 }
 
 void FBlackmagicMediaOutputModule::ShutdownModule()
