- Fixing crash on startup in Aja and Blackmagic in cooked projects

diff --git a/Engine/Plugins/Media/AjaMedia/Source/AjaMedia/Private/Aja/Aja.cpp b/Engine/Plugins/Media/AjaMedia/Source/AjaMedia/Private/Aja/Aja.cpp
--- a/Engine/Plugins/Media/AjaMedia/Source/AjaMedia/Private/Aja/Aja.cpp
+++ b/Engine/Plugins/Media/AjaMedia/Source/AjaMedia/Private/Aja/Aja.cpp
@@ -31,6 +31,10 @@ bool FAja::Initialize()
 	FPlatformProcess::PushDllDirectory(*AjaDllPath);
 	AjaDllPath = FPaths::Combine(AjaDllPath, AjaDll);
 
+	//Need to add the path to our dll dependency for cooked game, otherwise it won't be found when launching
+	FString GPUTextureTransferDllPath = FPaths::Combine(IPluginManager::Get().FindPlugin(TEXT("MediaIOFramework"))->GetBaseDir(), TEXT("/Binaries/Win64"));
+	FPlatformProcess::PushDllDirectory(*GPUTextureTransferDllPath);
+
 	if (!FPaths::FileExists(AjaDllPath))
 	{
 		UE_LOG(LogAjaMedia, Error, TEXT("Failed to find the binary folder for the AJA dll. Plug-in will not be functional."));
diff --git a/Engine/Plugins/Media/AjaMedia/Source/AjaMediaOutput/Private/AjaMediaOutputModule.cpp b/Engine/Plugins/Media/AjaMedia/Source/AjaMediaOutput/Private/AjaMediaOutputModule.cpp
--- a/Engine/Plugins/Media/AjaMedia/Source/AjaMediaOutput/Private/AjaMediaOutputModule.cpp
+++ b/Engine/Plugins/Media/AjaMedia/Source/AjaMediaOutput/Private/AjaMediaOutputModule.cpp
@@ -6,6 +6,7 @@
 #include "AjaLib.h"
 #include "DynamicRHI.h"
 #include "GenericPlatform/GenericPlatformDriver.h"
+#include "Misc/CoreDelegates.h"
 #include "RenderingThread.h"
 #include "RHI.h"
 
@@ -15,69 +16,82 @@ DEFINE_LOG_CATEGORY(LogAjaMediaOutput);
 
 void FAjaMediaOutputModule::StartupModule()
 {
-	const FGPUDriverInfo GPUDriverInfo = FPlatformMisc::GetGPUDriverInfo(GRHIAdapterName);
-	bIsGPUTextureTransferAvailable = GPUDriverInfo.IsNVIDIA();
-
-	if (bIsGPUTextureTransferAvailable)
+	const auto DMAInitializationFunc = [this]()
 	{
-		ENQUEUE_RENDER_COMMAND(AjaMediaCaptureInitialize)(
-			[this](FRHICommandListImmediate& RHICmdList) mutable
-			{
-				if (!GDynamicRHI)
-				{
-					return;
-				}
+		const FGPUDriverInfo GPUDriverInfo = FPlatformMisc::GetGPUDriverInfo(GRHIAdapterName);
+		bIsGPUTextureTransferAvailable = GPUDriverInfo.IsNVIDIA();
 
-				auto GetRHI = []()
+		if (bIsGPUTextureTransferAvailable)
+		{
+			ENQUEUE_RENDER_COMMAND(AjaMediaCaptureInitialize)(
+				[this](FRHICommandListImmediate& RHICmdList) mutable
 				{
-					FString RHIName = GDynamicRHI->GetName();
-					if (RHIName == TEXT("D3D11"))
+					if (!GDynamicRHI)
 					{
-						return AJA::ERHI::D3D11;
+						return;
 					}
-					else if (RHIName == TEXT("D3D12"))
+
+					auto GetRHI = []()
 					{
-						return AJA::ERHI::D3D12;
-					}
-					else if (RHIName == TEXT("Vulkan"))
+						FString RHIName = GDynamicRHI->GetName();
+						if (RHIName == TEXT("D3D11"))
+						{
+							return AJA::ERHI::D3D11;
+						}
+						else if (RHIName == TEXT("D3D12"))
+						{
+							return AJA::ERHI::D3D12;
+						}
+						else if (RHIName == TEXT("Vulkan"))
+						{
+							return AJA::ERHI::Vulkan;
+						}
+
+						return AJA::ERHI::Invalid;
+
+					};
+
+					AJA::FInitializeDMAArgs Args;
+					AJA::ERHI RHI = GetRHI();
+					Args.RHI = RHI;
+					/* Re-enable when adding vulkan support
+					if (RHI == AJA::ERHI::Vulkan)
 					{
-						return AJA::ERHI::Vulkan;
-					}
-
-					return AJA::ERHI::Invalid;
+						FVulkanDynamicRHI* vkDynamicRHI = static_cast<FVulkanDynamicRHI*>(GDynamicRHI);
+						Args.InVulkanInstance = vkDynamicRHI->GetInstance();
 
-				};
-
-				AJA::FInitializeDMAArgs Args;
-				AJA::ERHI RHI = GetRHI();
-				Args.RHI = RHI;
-				/* Re-enable when adding vulkan support
-				if (RHI == AJA::ERHI::Vulkan)
+						FMemory::Memcpy(Args.InRHIDeviceUUID, vkDynamicRHI->GetDevice()->GetDeviceIdProperties().deviceUUID, 16);
+					}
+					*/
+					Args.RHIDevice = GDynamicRHI->RHIGetNativeDevice();
+					Args.RHICommandQueue = GDynamicRHI->RHIGetNativeGraphicsQueue();
+
+					bIsGPUTextureTransferAvailable = AJA::InitializeDMA(Args);
+				});
+		}
+	};
+	
+	auto DMAUninitialize = [this]()
+	{
+		if (bIsGPUTextureTransferAvailable)
+		{
+			ENQUEUE_RENDER_COMMAND(AjaMediaCaptureUninitialize)(
+				[](FRHICommandListImmediate& RHICmdList) mutable
 				{
-					FVulkanDynamicRHI* vkDynamicRHI = static_cast<FVulkanDynamicRHI*>(GDynamicRHI);
-					Args.InVulkanInstance = vkDynamicRHI->GetInstance();
-
-					FMemory::Memcpy(Args.InRHIDeviceUUID, vkDynamicRHI->GetDevice()->GetDeviceIdProperties().deviceUUID, 16);
-				}
-				*/
-				Args.RHIDevice = GDynamicRHI->RHIGetNativeDevice();
-				Args.RHICommandQueue = GDynamicRHI->RHIGetNativeGraphicsQueue();
-
-				bIsGPUTextureTransferAvailable = AJA::InitializeDMA(Args);
-			});
-	}
+					AJA::UninitializeDMA();
+				});
+		}
+	};
+
+	//Postpone initialization after all modules have been loaded to be sure Aja library has been loaded
+	FCoreDelegates::OnAllModuleLoadingPhasesComplete.AddLambda(DMAInitializationFunc);
+	//Same for shutdown, uninitialize ourselves before library is unloaded
+	FCoreDelegates::OnEnginePreExit.AddLambda(DMAUninitialize);
 }
 
 void FAjaMediaOutputModule::ShutdownModule()
 {
-	ENQUEUE_RENDER_COMMAND(AjaMediaCaptureUninitialize)(
-		[bIsAvailable = bIsGPUTextureTransferAvailable](FRHICommandListImmediate& RHICmdList) mutable
-		{
-			if (bIsAvailable)
-			{
-				AJA::UninitializeDMA();
-			}
-		});
+
 }
 
 bool FAjaMediaOutputModule::IsGPUTextureTransferAvailable() const
diff --git a/Engine/Plugins/Media/BlackmagicMedia/Source/BlackmagicMedia/Private/Blackmagic/Blackmagic.cpp b/Engine/Plugins/Media/BlackmagicMedia/Source/BlackmagicMedia/Private/Blackmagic/Blackmagic.cpp
--- a/Engine/Plugins/Media/BlackmagicMedia/Source/BlackmagicMedia/Private/Blackmagic/Blackmagic.cpp
+++ b/Engine/Plugins/Media/BlackmagicMedia/Source/BlackmagicMedia/Private/Blackmagic/Blackmagic.cpp
@@ -31,6 +31,10 @@ bool FBlackmagic::Initialize()
 	FPlatformProcess::PushDllDirectory(*DllPath);
 	DllPath = FPaths::Combine(DllPath, VideoIODll);
 
+	//Need to add the path to our dll dependency for cooked game, otherwise it won't be found when launching
+	FString GPUTextureTransferDllPath = FPaths::Combine(IPluginManager::Get().FindPlugin(TEXT("MediaIOFramework"))->GetBaseDir(), TEXT("/Binaries/Win64"));
+	FPlatformProcess::PushDllDirectory(*GPUTextureTransferDllPath);
+
 	if (!FPaths::FileExists(DllPath))
 	{
 		UE_LOG(LogBlackmagicMedia, Error, TEXT("Failed to find the binary folder for the dll. Plug-in will not be functional."));
diff --git a/Engine/Plugins/Media/BlackmagicMedia/Source/BlackmagicMediaOutput/Private/BlackmagicMediaOutputModule.cpp b/Engine/Plugins/Media/BlackmagicMedia/Source/BlackmagicMediaOutput/Private/BlackmagicMediaOutputModule.cpp
--- a/Engine/Plugins/Media/BlackmagicMedia/Source/BlackmagicMediaOutput/Private/BlackmagicMediaOutputModule.cpp
+++ b/Engine/Plugins/Media/BlackmagicMedia/Source/BlackmagicMediaOutput/Private/BlackmagicMediaOutputModule.cpp
@@ -16,63 +16,76 @@ FBlackmagicMediaOutputModule& FBlackmagicMediaOutputModule::Get()
 
 void FBlackmagicMediaOutputModule::StartupModule()
 {
-	const FGPUDriverInfo GPUDriverInfo = FPlatformMisc::GetGPUDriverInfo(GRHIAdapterName);
-	bIsGPUTextureTransferAvailable = GPUDriverInfo.IsNVIDIA();
-
-	if (bIsGPUTextureTransferAvailable)
+	const auto DMAInitializationFunc = [this]()
 	{
-		ENQUEUE_RENDER_COMMAND(BlackmagicMediaCaptureInitialize)(
-			[this](FRHICommandListImmediate& RHICmdList) mutable
-			{
-				if (!GDynamicRHI)
-				{
-					return;
-				}
+		const FGPUDriverInfo GPUDriverInfo = FPlatformMisc::GetGPUDriverInfo(GRHIAdapterName);
+		bIsGPUTextureTransferAvailable = GPUDriverInfo.IsNVIDIA();
 
-				auto GetRHI = []()
+		if (bIsGPUTextureTransferAvailable)
+		{
+			ENQUEUE_RENDER_COMMAND(BlackmagicMediaCaptureInitialize)(
+				[this](FRHICommandListImmediate& RHICmdList) mutable
 				{
-					FString RHIName = GDynamicRHI->GetName();
-					if (RHIName == TEXT("D3D11"))
-					{
-						return BlackmagicDesign::ERHI::D3D11;
-					}
-					else if (RHIName == TEXT("D3D12"))
+					if (!GDynamicRHI)
 					{
-						return BlackmagicDesign::ERHI::D3D12;
+						return;
 					}
-					else if (RHIName == TEXT("Vulkan"))
+
+					auto GetRHI = []()
 					{
-						return BlackmagicDesign::ERHI::Vulkan;
-					}
+						FString RHIName = GDynamicRHI->GetName();
+						if (RHIName == TEXT("D3D11"))
+						{
+							return BlackmagicDesign::ERHI::D3D11;
+						}
+						else if (RHIName == TEXT("D3D12"))
+						{
+							return BlackmagicDesign::ERHI::D3D12;
+						}
+						else if (RHIName == TEXT("Vulkan"))
+						{
+							return BlackmagicDesign::ERHI::Vulkan;
+						}
+
+						return  BlackmagicDesign::ERHI::Invalid;
+					};
 
-					return  BlackmagicDesign::ERHI::Invalid;
-				};
 
+					BlackmagicDesign::FInitializeDMAArgs Args;
+					BlackmagicDesign::ERHI RHI = GetRHI();
+					Args.RHI = RHI;
+					Args.RHIDevice = GDynamicRHI->RHIGetNativeDevice();
+					Args.RHICommandQueue = GDynamicRHI->RHIGetNativeGraphicsQueue();
 
-				BlackmagicDesign::FInitializeDMAArgs Args;
-				BlackmagicDesign::ERHI RHI = GetRHI();
-				Args.RHI = RHI;
-				Args.RHIDevice = GDynamicRHI->RHIGetNativeDevice();
-				Args.RHICommandQueue = GDynamicRHI->RHIGetNativeGraphicsQueue();
+					bIsGPUTextureTransferAvailable = BlackmagicDesign::InitializeDMA(Args);
+				});
+		}
+	};
 
-				bIsGPUTextureTransferAvailable = BlackmagicDesign::InitializeDMA(Args);
-			});
-	}
+
+	auto DMAUninitialize = [this]()
+	{
+		if (bIsGPUTextureTransferAvailable)
+		{
+			ENQUEUE_RENDER_COMMAND(BlackmagicMediaCaptureUninitialize)(
+				[](FRHICommandListImmediate& RHICmdList) mutable
+				{
+					BlackmagicDesign::UninitializeDMA();
+				});
+		}
+	};
+
+	//Postpone initialization after all modules have been loaded to be sure Blackmagic library has been loaded
+	FCoreDelegates::OnAllModuleLoadingPhasesComplete.AddLambda(DMAInitializationFunc);
+	//Same for shutdown, uninitialize ourselves before library is unloaded
+	FCoreDelegates::OnEnginePreExit.AddLambda(DMAUninitialize);
 }
 
 void FBlackmagicMediaOutputModule::ShutdownModule()
 {
-	ENQUEUE_RENDER_COMMAND(BlackmagicMediaCaptureUninitialize)(
-		[bIsAvailable = bIsGPUTextureTransferAvailable](FRHICommandListImmediate& RHICmdList) mutable
-		{
-			if (bIsAvailable)
-			{
-				BlackmagicDesign::UninitializeDMA();
-			}
-		});
+	
 }
 
-
 bool FBlackmagicMediaOutputModule::IsGPUTextureTransferAvailable() const
 {
 	return bIsGPUTextureTransferAvailable;
