Aja GPU direct input, 10 bit input support, more vulkan fixes

diff --git a/Engine/Plugins/Media/AjaMedia/Source/AjaMedia/AjaMedia.Build.cs b/Engine/Plugins/Media/AjaMedia/Source/AjaMedia/AjaMedia.Build.cs
--- a/Engine/Plugins/Media/AjaMedia/Source/AjaMedia/AjaMedia.Build.cs
+++ b/Engine/Plugins/Media/AjaMedia/Source/AjaMedia/AjaMedia.Build.cs
@@ -23,11 +23,13 @@ namespace UnrealBuildTool.Rules
 					"Core",
 					"CoreUObject",
 					"Engine",
+					"GPUTextureTransfer",
 					"MediaIOCore",
 					"MediaUtils",
 					"Projects",
 					"TimeManagement",
 					"RenderCore",
+					"RHI",
 					"SlateCore"
 				});
 
diff --git a/Engine/Plugins/Media/AjaMedia/Source/AjaMedia/Private/Player/AjaMediaPlayer.cpp b/Engine/Plugins/Media/AjaMedia/Source/AjaMedia/Private/Player/AjaMediaPlayer.cpp
--- a/Engine/Plugins/Media/AjaMedia/Source/AjaMedia/Private/Player/AjaMediaPlayer.cpp
+++ b/Engine/Plugins/Media/AjaMedia/Source/AjaMedia/Private/Player/AjaMediaPlayer.cpp
@@ -7,13 +7,15 @@
 #include "MediaIOCoreEncodeTime.h"
 #include "MediaIOCoreFileWriter.h"
 #include "MediaIOCoreSamples.h"
-
+#include "GPUTextureTransferModule.h"
 #include "HAL/PlatformAtomics.h"
 #include "HAL/PlatformProcess.h"
 #include "IAjaMediaModule.h"
 #include "IMediaEventSink.h"
 #include "IMediaOptions.h"
+
 #include "Misc/ScopeLock.h"
+#include "RenderCommandFence.h"
 #include "Stats/Stats2.h"
 #include "Styling/SlateStyle.h"
 
@@ -22,8 +24,6 @@
 #include "AjaMediaSettings.h"
 #include "AjaMediaTextureSample.h"
 
-#include "AjaMediaAllowPlatformTypes.h"
-
 #if WITH_EDITOR
 #include "EngineAnalytics.h"
 #endif
@@ -46,14 +46,19 @@ static FAutoConsoleCommand AjaWriteOutputRawDataCmd(
 	FConsoleCommandDelegate::CreateLambda([]() { bAjaWriteOutputRawDataCmdEnable = true; })
 	);
 
+static TAutoConsoleVariable<int32> CVarAjaEnableGPUDirect(
+	TEXT("Aja.EnableGPUDirectInput"), 0,
+	TEXT("Whether to enable GPU direct for faster video frame copies. (Experimental)"),
+	ECVF_RenderThreadSafe);
+
 /* FAjaVideoPlayer structors
  *****************************************************************************/
 
 FAjaMediaPlayer::FAjaMediaPlayer(IMediaEventSink& InEventSink)
 	: Super(InEventSink)
-	, AudioSamplePool(new FAjaMediaAudioSamplePool)
-	, MetadataSamplePool(new FAjaMediaBinarySamplePool)
-	, TextureSamplePool(new FAjaMediaTextureSamplePool)
+	, AudioSamplePool(MakeUnique<FAjaMediaAudioSamplePool>())
+	, MetadataSamplePool(MakeUnique<FAjaMediaBinarySamplePool>())
+	, TextureSamplePool(MakeUnique<FAjaMediaTextureSamplePool>())
 	, MaxNumAudioFrameBuffer(8)
 	, MaxNumMetadataFrameBuffer(8)
 	, MaxNumVideoFrameBuffer(8)
@@ -75,15 +80,14 @@ FAjaMediaPlayer::FAjaMediaPlayer(IMediaEventSink& InEventSink)
 	, SupportedSampleTypes(EMediaIOSampleType::None)
 	, bPauseRequested(false)
 {
+	FGPUTextureTransferModule& Module = FGPUTextureTransferModule::Get();
+	GPUTextureTransfer = Module.GetTextureTransfer();
 }
 
 
 FAjaMediaPlayer::~FAjaMediaPlayer()
 {
 	Close();
-	delete AudioSamplePool;
-	delete MetadataSamplePool;
-	delete TextureSamplePool;
 }
 
 
@@ -213,6 +217,7 @@ bool FAjaMediaPlayer::Open(const FString& Url, const IMediaOptions* Options)
 	MaxNumAudioFrameBuffer = Options->GetMediaOption(AjaMediaOption::MaxAudioFrameBuffer, (int64)8);
 	MaxNumMetadataFrameBuffer = Options->GetMediaOption(AjaMediaOption::MaxAncillaryFrameBuffer, (int64)8);
 	MaxNumVideoFrameBuffer = Options->GetMediaOption(AjaMediaOption::MaxVideoFrameBuffer, (int64)8);
+	PixelFormat = (EAjaMediaSourceColorFormat)(Options->GetMediaOption(AjaMediaOption::ColorFormat, (int64)EAjaMediaSourceColorFormat::YUV2_8bit));
 
 	check(InputChannel == nullptr);
 	InputChannel = new AJA::AJAInputChannel();
@@ -254,7 +259,13 @@ bool FAjaMediaPlayer::Open(const FString& Url, const IMediaOptions* Options)
 		FEngineAnalytics::GetProvider().RecordEvent(TEXT("MediaFramework.AjaSourceOpened"), EventAttributes);
 	}
 #endif
-	
+
+
+	if (CVarAjaEnableGPUDirect.GetValueOnAnyThread())
+	{
+		CreateAndRegisterTextures(Options);
+	}
+
 	return true;
 }
 
@@ -282,8 +293,11 @@ void FAjaMediaPlayer::Close()
 	AjaThreadCurrentTextureSample.Reset();
 
 	Super::Close();
-}
 
+	// This has to be done after the Super::Close to make sure that samples are flushed so that all textures are returned to our texturex array.
+	UnregisterSampleBuffers();
+	UnregisterTextures();
+}
 
 FGuid FAjaMediaPlayer::GetPlayerPluginGUID() const
 {
@@ -478,6 +492,7 @@ void FAjaMediaPlayer::OnCompletion(bool bSucceed)
 
 bool FAjaMediaPlayer::OnRequestInputBuffer(const AJA::AJARequestInputBufferData& InRequestBuffer, AJA::AJARequestedInputBufferData& OutRequestedBuffer)
 {
+
 	SCOPE_CYCLE_COUNTER(STAT_AJA_MediaPlayer_RequestFrame);
 
 	// Do not request a video buffer if the frame is interlaced. We need 2 samples and we need to process them.
@@ -512,8 +527,43 @@ bool FAjaMediaPlayer::OnRequestInputBuffer(const AJA::AJARequestInputBufferData&
 	// Video
 	if (bUseVideo && InRequestBuffer.VideoBufferSize > 0 && InRequestBuffer.bIsProgressivePicture)
 	{
+		if (CVarAjaEnableGPUDirect.GetValueOnAnyThread())
+		{
+			if (!Textures.Num())
+			{
+				UE_LOG(LogAjaMedia, Error, TEXT("No texture available while doing a gpu texture transfer."));
+				return false;
+			}
+		}
+
 		AjaThreadCurrentTextureSample = TextureSamplePool->AcquireShared();
-		OutRequestedBuffer.VideoBuffer = reinterpret_cast<uint8_t*>(AjaThreadCurrentTextureSample->RequestBuffer(InRequestBuffer.VideoBufferSize));
+		AjaThreadCurrentTextureSample->RequestBuffer(InRequestBuffer.VideoBufferSize);
+
+		if (CVarAjaEnableGPUDirect.GetValueOnAnyThread())
+		{
+			bool bFoundRegisteredTexture = false;
+
+			do
+			{
+				FScopeLock Lock(&TexturesCriticalSection);
+				TRefCountPtr<FRHITexture> Texture = Textures.Pop();
+				if (Texture && RegisteredTextures.Contains(Texture))
+				{
+					AjaThreadCurrentTextureSample->SetTexture(Texture);
+					AjaThreadCurrentTextureSample->SetDestructionCallback([this](TRefCountPtr<FRHITexture> InTexture) { OnSampleDestroyed(InTexture); });
+					bFoundRegisteredTexture = true;
+				}
+				else
+				{
+					UE_LOG(LogAjaMedia, Display, TEXT("Unregistered texture %u encountered while doing a gpu texture transfer."), Texture ? reinterpret_cast<std::uintptr_t>(Texture->GetNativeResource()) : 0);
+				}
+
+			} while (!bFoundRegisteredTexture && Textures.Num());
+
+			RegisterSampleBuffer(AjaThreadCurrentTextureSample);
+		}
+
+		OutRequestedBuffer.VideoBuffer = reinterpret_cast<uint8_t*>(AjaThreadCurrentTextureSample->GetMutableBuffer());
 	}
 
 	return true;
@@ -666,7 +716,7 @@ bool FAjaMediaPlayer::OnInputFrameReceived(const AJA::AJAInputFrameData& InInput
 			break;
 		}
 
-		if (bEncodeTimecodeInTexel && DecodedTimecode.IsSet() && InVideoFrame.bIsProgressivePicture)
+		if (bEncodeTimecodeInTexel && DecodedTimecode.IsSet() && InVideoFrame.bIsProgressivePicture && !CVarAjaEnableGPUDirect.GetValueOnAnyThread())
 		{
 			FTimecode SetTimecode = DecodedTimecode.GetValue();
 			FMediaIOCoreEncodeTime EncodeTime(EncodePixelFormat, InVideoFrame.VideoBuffer, InVideoFrame.Stride, InVideoFrame.Width, InVideoFrame.Height);
@@ -683,7 +733,30 @@ bool FAjaMediaPlayer::OnInputFrameReceived(const AJA::AJAInputFrameData& InInput
 		{
 			if (AjaThreadCurrentTextureSample->SetProperties(InVideoFrame.Stride, InVideoFrame.Width, InVideoFrame.Height, VideoSampleFormat, DecodedTime, VideoFrameRate, DecodedTimecode, bIsSRGBInput))
 			{
-				Samples->AddVideo(AjaThreadCurrentTextureSample.ToSharedRef());
+				if (GPUTextureTransfer && CVarAjaEnableGPUDirect.GetValueOnAnyThread())
+				{
+					// We must start transfer on render thread
+					ENQUEUE_RENDER_COMMAND(StartGPUTextureTransferCopy)(
+						[this, CurrentSample = AjaThreadCurrentTextureSample](FRHICommandListImmediate& RHICmdList)
+						{
+							if (CurrentSample->GetTexture())
+							{
+								if (GPUTextureTransfer)
+								{
+									void* Buffer = CurrentSample->GetMutableBuffer();
+									UE_LOG(LogAjaMedia, Verbose, TEXT("Starting Transfer with buffer %u"), reinterpret_cast<std::uintptr_t>(Buffer));
+									GPUTextureTransfer->TransferTexture(Buffer, CurrentSample->GetTexture(), UE::GPUTextureTransfer::ETransferDirection::CPU_TO_GPU);
+									GPUTextureTransfer->BeginSync((CurrentSample)->GetMutableBuffer(), UE::GPUTextureTransfer::ETransferDirection::CPU_TO_GPU);
+									Samples->AddVideo(CurrentSample.ToSharedRef());
+									GPUTextureTransfer->EndSync(CurrentSample->GetMutableBuffer());
+								}
+							}
+						});
+				}
+				else
+				{
+					Samples->AddVideo(AjaThreadCurrentTextureSample.ToSharedRef());
+				}
 			}
 		}
 		else
@@ -750,6 +823,176 @@ void FAjaMediaPlayer::SetupSampleChannels()
 	Samples->InitializeMetadataBuffer(MetadataSettings);
 }
 
+void FAjaMediaPlayer::OnSampleDestroyed(TRefCountPtr<FRHITexture> InTexture)
+{
+	FScopeLock ScopeLock(&TexturesCriticalSection);
+	if (InTexture)
+	{
+		Textures.Add(InTexture);
+	}
+}
+
+void FAjaMediaPlayer::RegisterSampleBuffer(const TSharedPtr<FAjaMediaTextureSample>& InSample)
+{
+	if (!RegisteredBuffers.Contains(InSample->GetMutableBuffer()))
+	{
+		RegisteredBuffers.Add(InSample->GetMutableBuffer());
+
+		// Enqueue buffer registration on the render  thread to ensure it happens before the call to TransferTexture
+		ENQUEUE_RENDER_COMMAND(RegisterSampleBuffers)([this, InSample](FRHICommandList& CommandList)
+			{
+				if (!InSample->GetMutableBuffer())
+				{
+					UE_LOG(LogAjaMedia, Error, TEXT("A buffer was not available while performing a gpu texture transfer."));
+					return;
+				}
+
+				if (!InSample->GetTexture())
+				{
+					UE_LOG(LogAjaMedia, Error, TEXT("A texture was not available while performing a gpu texture transfer."));
+					return;
+				}
+
+				const uint32 TextureWidth = InSample->GetTexture()->GetDesc().Extent.X;
+				const uint32 TextureHeight = InSample->GetTexture()->GetDesc().Extent.Y;
+				uint32 TextureStride = TextureWidth * 4;
+
+				EPixelFormat Format = InSample->GetTexture()->GetFormat();
+				if (Format == PF_R32G32B32A32_UINT)
+				{
+					TextureStride *= 4;
+				}
+
+				UE_LOG(LogAjaMedia, Verbose, TEXT("Registering buffer %u"), reinterpret_cast<std::uintptr_t>(InSample->GetMutableBuffer()));
+				UE::GPUTextureTransfer::FRegisterDMABufferArgs Args;
+				Args.Buffer = InSample->GetMutableBuffer();
+				Args.Width = TextureWidth;
+				Args.Height = TextureHeight;
+				Args.Stride = TextureStride;
+				Args.PixelFormat = Format == PF_B8G8R8A8 ? UE::GPUTextureTransfer::EPixelFormat::PF_8Bit : UE::GPUTextureTransfer::EPixelFormat::PF_10Bit;
+				GPUTextureTransfer->RegisterBuffer(Args);
+			});
+	}
+	else
+	{
+		UE_LOG(LogAjaMedia, VeryVerbose, TEXT("Buffer already registered: %u"), reinterpret_cast<std::uintptr_t>(AjaThreadCurrentTextureSample->GetMutableBuffer()));
+	}
+}
+
+void FAjaMediaPlayer::UnregisterSampleBuffers()
+{
+	ENQUEUE_RENDER_COMMAND(UnregisterSampleBuffers)([TextureTransfer =GPUTextureTransfer, BuffersToUnregister = RegisteredBuffers](FRHICommandList& CommandList)
+	{
+		if (TextureTransfer)
+		{
+			for (void* RegisteredBuffer : BuffersToUnregister)
+			{
+				TextureTransfer->UnregisterBuffer(RegisteredBuffer);
+			}
+		}
+	});
+	
+	RegisteredBuffers.Reset();
+}
+
+void FAjaMediaPlayer::CreateAndRegisterTextures(const IMediaOptions* Options)
+{
+	ENQUEUE_RENDER_COMMAND(AjaMediaPlayerCreateTextures)(
+		[this, Options](FRHICommandListImmediate& CommandList)
+		{
+			if (!GPUTextureTransfer)
+			{
+				UE_LOG(LogAjaMedia, Error, TEXT("Could not register textures for gpu texture transfer, as the GPUTextureTransfer object was null."));
+				return;
+			}
+
+			for (uint8 Index = 0; Index < MaxNumVideoFrameBuffer; Index++)
+			{
+				TRefCountPtr<FRHITexture> RHITexture;
+				const FString TextureName = FString::Printf(TEXT("FAjaMediaPlayerTexture %d"), Index);
+				FRHIResourceCreateInfo CreateInfo(*TextureName);
+				
+				uint32 TextureWidth = VideoTrackFormat.Dim.X;
+				const uint32 TextureHeight = VideoTrackFormat.Dim.Y;
+				uint32 Stride = TextureWidth;
+				EPixelFormat InputFormat = EPixelFormat::PF_Unknown;
+
+				if (PixelFormat == EAjaMediaSourceColorFormat::YUV2_8bit)
+				{
+					InputFormat = EPixelFormat::PF_B8G8R8A8;
+					TextureWidth /= 2;
+					Stride = TextureWidth * 4;
+				}
+				else if (PixelFormat == EAjaMediaSourceColorFormat::YUV_10bit)
+				{
+					InputFormat = EPixelFormat::PF_R32G32B32A32_UINT;
+					TextureWidth /= 6;
+					Stride = TextureWidth * 16;
+				}
+				else
+				{
+					checkf(false, TEXT("Format not supported"));
+				}
+
+				TRefCountPtr<FRHITexture2D> DummyTexture2DRHI;
+
+				ETextureCreateFlags CreateFlags = TexCreate_Shared;
+				if (RHIGetInterfaceType() == ERHIInterfaceType::Vulkan)
+				{
+					CreateFlags = TexCreate_External;
+				}
+
+				FRHITextureCreateDesc CreateDesc = FRHITextureCreateDesc::Create2D(*TextureName, FIntPoint(TextureWidth, TextureHeight), InputFormat, FClearValueBinding::White, CreateFlags);
+				RHICreateTargetableShaderResource(CreateDesc, TexCreate_RenderTargetable, RHITexture, DummyTexture2DRHI);
+
+				UE_LOG(LogAjaMedia, Verbose, TEXT("Registering texture %u"), reinterpret_cast<std::uintptr_t>(RHITexture->GetNativeResource()));
+
+				UE::GPUTextureTransfer::FRegisterDMATextureArgs Args;
+				Args.RHITexture = RHITexture.GetReference();
+				Args.Stride = Stride;
+				Args.Width = TextureWidth;
+				Args.Height = TextureHeight;
+				Args.RHIResourceMemory = nullptr;
+				Args.PixelFormat = PixelFormat == EAjaMediaSourceColorFormat::YUV2_8bit ? UE::GPUTextureTransfer::EPixelFormat::PF_8Bit : UE::GPUTextureTransfer::EPixelFormat::PF_10Bit;
+				
+				GPUTextureTransfer->RegisterTexture(Args);
+				RegisteredTextures.Add(RHITexture);
+
+				{
+					FScopeLock Lock(&TexturesCriticalSection);
+					Textures.Add(RHITexture);
+				}
+			}
+		});
+
+	FRenderCommandFence RenderFence;
+	RenderFence.BeginFence();
+	RenderFence.Wait();
+}
+
+void FAjaMediaPlayer::UnregisterTextures()
+{
+	// Keep a copy of the textures to make sure they were not destoryed before this gets executed.
+	ENQUEUE_RENDER_COMMAND(UnregisterSampleBuffers)([TextureTransfer = GPUTextureTransfer, TexturesToUnregister = RegisteredTextures](FRHICommandList& CommandList)
+		{
+			if (TextureTransfer)
+			{
+				for (const TRefCountPtr<FRHITexture>& RHITexture : TexturesToUnregister)
+				{
+					UE_LOG(LogAjaMedia, Verbose, TEXT("Unregistering texture %u"), reinterpret_cast<std::uintptr_t>(RHITexture->GetNativeResource()));
+					TextureTransfer->UnregisterTexture(RHITexture.GetReference());
+				}
+			}
+		});
+
+	RegisteredTextures.Reset();
+
+	{
+		FScopeLock Lock(&TexturesCriticalSection);
+		Textures.Reset();
+	}
+}
+
 bool FAjaMediaPlayer::SetRate(float Rate)
 {
 	if (FMath::IsNearlyEqual(Rate, 1.0f))
@@ -768,5 +1011,3 @@ bool FAjaMediaPlayer::SetRate(float Rate)
 }
 
 #undef LOCTEXT_NAMESPACE
-
-#include "AjaMediaHidePlatformTypes.h"
diff --git a/Engine/Plugins/Media/AjaMedia/Source/AjaMedia/Private/Player/AjaMediaPlayer.h b/Engine/Plugins/Media/AjaMedia/Source/AjaMedia/Private/Player/AjaMediaPlayer.h
--- a/Engine/Plugins/Media/AjaMedia/Source/AjaMedia/Private/Player/AjaMediaPlayer.h
+++ b/Engine/Plugins/Media/AjaMedia/Source/AjaMedia/Private/Player/AjaMediaPlayer.h
@@ -6,6 +6,10 @@
 
 #include "AjaMediaPrivate.h"
 #include "AjaMediaSource.h"
+#include "GPUTextureTransferModule.h"
+#include "GPUTextureTransfer.h"
+#include "HAL/CriticalSection.h"
+#include "Templates/PimplPtr.h"
 
 #include <atomic>
 
@@ -103,12 +107,18 @@ protected:
 	//~ FMediaIOCorePlayerBase interface
 	virtual void SetupSampleChannels() override;
 
+private:
+	void OnSampleDestroyed(TRefCountPtr<FRHITexture> InTexture);
+	void RegisterSampleBuffer(const TSharedPtr<FAjaMediaTextureSample>& InSample);
+	void UnregisterSampleBuffers();
+	void CreateAndRegisterTextures(const IMediaOptions* Options);
+	void UnregisterTextures();
 private:
 
 	/** Audio, MetaData, Texture  sample object pool. */
-	FAjaMediaAudioSamplePool* AudioSamplePool;
-	FAjaMediaBinarySamplePool* MetadataSamplePool;
-	FAjaMediaTextureSamplePool* TextureSamplePool;
+	TUniquePtr<FAjaMediaAudioSamplePool> AudioSamplePool;
+	TUniquePtr<FAjaMediaBinarySamplePool> MetadataSamplePool;
+	TUniquePtr<FAjaMediaTextureSamplePool> TextureSamplePool;
 
 	TSharedPtr<FMediaIOCoreBinarySampleBase, ESPMode::ThreadSafe> AjaThreadCurrentAncSample;
 	TSharedPtr<FMediaIOCoreBinarySampleBase, ESPMode::ThreadSafe> AjaThreadCurrentAncF2Sample;
@@ -171,4 +181,18 @@ private:
 
 	/** Flag to indicate that pause is being requested */
 	std::atomic<bool> bPauseRequested;
+
+	FCriticalSection TexturesCriticalSection;
+	
+	/** GPU Texture transfer object */
+	UE::GPUTextureTransfer::TextureTransferPtr GPUTextureTransfer;
+
+	/** Pool of textures registerd with GPU Texture transfer. */
+	TArray<TRefCountPtr<FRHITexture>> Textures;
+	/** Buffers Registered with GPU Texture Transfer */
+	TSet<void*> RegisteredBuffers;
+	/** Pool of textures registerd with GPU Texture transfer. */
+	TSet<TRefCountPtr<FRHITexture>> RegisteredTextures;
+
+	EAjaMediaSourceColorFormat PixelFormat = EAjaMediaSourceColorFormat::YUV2_8bit;
 };
diff --git a/Engine/Plugins/Media/AjaMedia/Source/AjaMedia/Private/Player/AjaMediaTextureSample.h b/Engine/Plugins/Media/AjaMedia/Source/AjaMedia/Private/Player/AjaMediaTextureSample.h
--- a/Engine/Plugins/Media/AjaMedia/Source/AjaMedia/Private/Player/AjaMediaTextureSample.h
+++ b/Engine/Plugins/Media/AjaMedia/Source/AjaMedia/Private/Player/AjaMediaTextureSample.h
@@ -15,6 +15,16 @@ class FAjaMediaTextureSample
 
 public:
 
+	virtual void ShutdownPoolable() override
+	{
+		if (DestructionCallback)
+		{
+			DestructionCallback(Texture);
+		}
+
+		Super::FreeSample();
+	}
+
 	/**
 	 * Initialize the sample.
 	 *
@@ -65,6 +75,16 @@ public:
 			, bInIsSRGB);
 	}
 
+	void SetTexture(TRefCountPtr<FRHITexture> InRHITexture)
+	{
+		Texture = MoveTemp(InRHITexture);
+	}
+
+	void SetDestructionCallback(TFunction<void(TRefCountPtr<FRHITexture2D>)> InDestructionCallback)
+	{
+		DestructionCallback = InDestructionCallback;
+	}
+
 	/**
 	 * Get YUV to RGB conversion matrix
 	 *
@@ -75,6 +95,39 @@ public:
 		return MediaShaders::YuvToRgbRec709Scaled;
 	}
 
+	virtual const void* GetBuffer() override
+	{
+		// Don't return the buffer if we have a texture to force the media player to use the texture if available. 
+		if (Texture)
+		{
+			return nullptr;
+		}
+		return Buffer.GetData();
+
+	}
+
+	void* GetMutableBuffer()
+	{
+		return Buffer.GetData();
+	}
+
+#if WITH_ENGINE
+	virtual FRHITexture* GetTexture() const override
+	{
+		if (Texture)
+		{
+			return Texture.GetReference();
+		}
+		return nullptr;
+	}
+#endif //WITH_ENGINE
+
+private:
+	/** Hold a texture to be used for gpu texture transfers. */
+	TRefCountPtr<FRHITexture2D> Texture;
+
+	/** Called when the sample is destroyed by its pool. */
+	TFunction<void(TRefCountPtr<FRHITexture2D>)> DestructionCallback;
 };
 
 /*
diff --git a/Engine/Plugins/Media/AjaMedia/Source/AjaMediaOutput/Private/AjaMediaCapture.cpp b/Engine/Plugins/Media/AjaMedia/Source/AjaMediaOutput/Private/AjaMediaCapture.cpp
--- a/Engine/Plugins/Media/AjaMedia/Source/AjaMediaOutput/Private/AjaMediaCapture.cpp
+++ b/Engine/Plugins/Media/AjaMedia/Source/AjaMediaOutput/Private/AjaMediaCapture.cpp
@@ -21,7 +21,7 @@
 #include "EngineAnalytics.h"
 #endif
 
-TAutoConsoleVariable<int32> CVarAjaEnableGPUDirect(
+static TAutoConsoleVariable<int32> CVarAjaEnableGPUDirect(
 	TEXT("Aja.EnableGPUDirect"), 0,
 	TEXT("Whether to enable GPU direct for faster video frame copies. (Experimental)"),
 	ECVF_RenderThreadSafe);
diff --git a/Engine/Plugins/Media/MediaIOFramework/MediaIOFramework.uplugin b/Engine/Plugins/Media/MediaIOFramework/MediaIOFramework.uplugin
--- a/Engine/Plugins/Media/MediaIOFramework/MediaIOFramework.uplugin
+++ b/Engine/Plugins/Media/MediaIOFramework/MediaIOFramework.uplugin
@@ -27,8 +27,8 @@
 		{
 			"Name": "GPUTextureTransfer",
 			"Type": "Runtime",
-			"LoadingPhase": "Default",
-			"PlatformAllowList": [ "Win64" ]
+			"LoadingPhase": "PostConfigInit", 
+			"PlatformAllowList": [ "Win64", "Linux" ]
 		}
 	],
 	"DocsURL": "",
diff --git a/Engine/Plugins/Media/MediaIOFramework/Source/GPUTextureTransfer/GPUTextureTransfer.Build.cs b/Engine/Plugins/Media/MediaIOFramework/Source/GPUTextureTransfer/GPUTextureTransfer.Build.cs
--- a/Engine/Plugins/Media/MediaIOFramework/Source/GPUTextureTransfer/GPUTextureTransfer.Build.cs
+++ b/Engine/Plugins/Media/MediaIOFramework/Source/GPUTextureTransfer/GPUTextureTransfer.Build.cs
@@ -11,13 +11,18 @@ namespace UnrealBuildTool.Rules
 				AddEngineThirdPartyPrivateStaticDependencies(Target, "GPUDirect");
 				AddEngineThirdPartyPrivateStaticDependencies(Target, "Vulkan");
 			}
-
+			else if (Target.Platform == UnrealTargetPlatform.Linux)
+			{
+				AddEngineThirdPartyPrivateStaticDependencies(Target, "Vulkan");
+			}
+				
 			PublicDefinitions.Add("PERF_LOGGING=0");
 
 			PublicDependencyModuleNames.AddRange(
 				new string[]
 				{
 					"Core",
+					"RHI"
 				});
 
 			PublicIncludePaths.AddRange(
diff --git a/Engine/Plugins/Media/MediaIOFramework/Source/GPUTextureTransfer/Private/D3D11TextureTransfer.cpp b/Engine/Plugins/Media/MediaIOFramework/Source/GPUTextureTransfer/Private/D3D11TextureTransfer.cpp
--- a/Engine/Plugins/Media/MediaIOFramework/Source/GPUTextureTransfer/Private/D3D11TextureTransfer.cpp
+++ b/Engine/Plugins/Media/MediaIOFramework/Source/GPUTextureTransfer/Private/D3D11TextureTransfer.cpp
@@ -5,15 +5,12 @@
 #if PLATFORM_WINDOWS
 #include "Windows/AllowWindowsPlatformTypes.h"
 #include "Windows/PreWindowsApi.h"
-#endif
 
 #include "DVPAPI.h"
 #include "dvpapi_d3d11.h"
 
-#if PLATFORM_WINDOWS
 #include "Windows/PostWindowsApi.h"
 #include "Windows/HideWindowsPlatformTypes.h"
-#endif
 
 namespace UE::GPUTextureTransfer::Private
 {
@@ -54,7 +51,7 @@ DVPStatus FD3D11TextureTransfer::BindBuffer_Impl(DVPBufferHandle InBufferHandle)
 	return dvpBindToD3D11Device(InBufferHandle, D3DDevice);
 }
 
-DVPStatus FD3D11TextureTransfer::CreateGPUResource_Impl(void* InTexture, FTextureInfo* OutTextureInfo) const
+DVPStatus FD3D11TextureTransfer::CreateGPUResource_Impl(const FRegisterDMATextureArgs& InArgs, FTextureInfo* OutTextureInfo) const
 {
 	if (!OutTextureInfo || !D3DDevice)
 	{
@@ -62,7 +59,7 @@ DVPStatus FD3D11TextureTransfer::CreateGPUResource_Impl(void* InTexture, FTextur
 	}
 
 	OutTextureInfo->External.Handle = nullptr;
-	return dvpCreateGPUD3D11Resource((ID3D11Resource*)InTexture, &OutTextureInfo->DVPHandle);
+	return dvpCreateGPUD3D11Resource((ID3D11Resource*)InArgs.RHITexture->GetNativeResource(), &OutTextureInfo->DVPHandle);
 }
 
 DVPStatus FD3D11TextureTransfer::UnbindBuffer_Impl(DVPBufferHandle InBufferHandle) const
@@ -75,4 +72,6 @@ DVPStatus FD3D11TextureTransfer::UnbindBuffer_Impl(DVPBufferHandle InBufferHandl
 	return dvpUnbindFromD3D11Device(InBufferHandle, D3DDevice);
 }
 
-}
\ No newline at end of file
+}
+
+#endif // PLATFORM_WINDOWS
diff --git a/Engine/Plugins/Media/MediaIOFramework/Source/GPUTextureTransfer/Private/D3D11TextureTransfer.h b/Engine/Plugins/Media/MediaIOFramework/Source/GPUTextureTransfer/Private/D3D11TextureTransfer.h
--- a/Engine/Plugins/Media/MediaIOFramework/Source/GPUTextureTransfer/Private/D3D11TextureTransfer.h
+++ b/Engine/Plugins/Media/MediaIOFramework/Source/GPUTextureTransfer/Private/D3D11TextureTransfer.h
@@ -2,6 +2,8 @@
 
 #pragma once
 
+#if PLATFORM_WINDOWS
+
 #include "TextureTransferBase.h"
 
 struct ID3D11Device;
@@ -14,11 +16,13 @@ protected:
 	virtual DVPStatus Init_Impl(const FInitializeDMAArgs& InArgs) override;
 	virtual DVPStatus GetConstants_Impl(uint32* OutBufferAddrAlignment, uint32* OutBufferGPUStrideAlignment, uint32* OutSemaphoreAddrAlignment, uint32* OutSemaphoreAllocSize, uint32* OutSemaphorePayloadOffset, uint32* OutSemaphorePayloadSize) const override;
 	virtual DVPStatus BindBuffer_Impl(DVPBufferHandle InBufferHandle) const override;
-	virtual DVPStatus CreateGPUResource_Impl(void* InTexture, FTextureInfo* OutTextureInfo) const override;
+	virtual DVPStatus CreateGPUResource_Impl(const FRegisterDMATextureArgs& InArgs, FTextureInfo* OutTextureInfo) const override;
 	virtual DVPStatus UnbindBuffer_Impl(DVPBufferHandle InBufferHandle) const override;
 	virtual DVPStatus CloseDevice_Impl() const override;
 
 private:
 	ID3D11Device* D3DDevice = nullptr;
 };
-}
\ No newline at end of file
+}
+
+#endif // PLATFORM_WINDOWS
\ No newline at end of file
diff --git a/Engine/Plugins/Media/MediaIOFramework/Source/GPUTextureTransfer/Private/D3D12TextureTransfer.cpp b/Engine/Plugins/Media/MediaIOFramework/Source/GPUTextureTransfer/Private/D3D12TextureTransfer.cpp
--- a/Engine/Plugins/Media/MediaIOFramework/Source/GPUTextureTransfer/Private/D3D12TextureTransfer.cpp
+++ b/Engine/Plugins/Media/MediaIOFramework/Source/GPUTextureTransfer/Private/D3D12TextureTransfer.cpp
@@ -5,15 +5,12 @@
 #if PLATFORM_WINDOWS
 #include "Windows/AllowWindowsPlatformTypes.h"
 #include "Windows/PreWindowsApi.h"
-#endif
 
 #include "DVPAPI.h"
 #include "dvpapi_d3d12.h"
 
-#if PLATFORM_WINDOWS
 #include "Windows/PostWindowsApi.h"
 #include "Windows/HideWindowsPlatformTypes.h"
-#endif
 
 
 namespace UE::GPUTextureTransfer::Private
@@ -56,14 +53,14 @@ namespace UE::GPUTextureTransfer::Private
 		return dvpBindToD3D12Device(InBufferHandle, D3DDevice);
 	}
 
-	DVPStatus FD3D12TextureTransfer::CreateGPUResource_Impl(void* InTexture, FTextureTransferBase::FTextureInfo* OutTextureInfo) const
+	DVPStatus FD3D12TextureTransfer::CreateGPUResource_Impl(const FRegisterDMATextureArgs& InArgs, FTextureTransferBase::FTextureInfo* OutTextureInfo) const
 	{
-		if (!OutTextureInfo || !D3DDevice)
+		if (!OutTextureInfo || !D3DDevice || !InArgs.Stride)
 		{
 			return DVP_STATUS_ERROR;
 		}
 
-		ID3D12Resource* D3D12Texture = (ID3D12Resource*) InTexture;
+		ID3D12Resource* D3D12Texture = (ID3D12Resource*) InArgs.RHITexture->GetNativeResource();
 		D3D12_RESOURCE_DESC ResourceDesc = D3D12Texture->GetDesc();
 
 		HRESULT Res = D3DDevice->CreateSharedHandle(D3D12Texture, NULL, GENERIC_ALL, NULL, &OutTextureInfo->External.Handle);
@@ -73,11 +70,11 @@ namespace UE::GPUTextureTransfer::Private
 		}
 
 		DVPGpuExternalResourceDesc Desc;
-		Desc.width = (uint32) ResourceDesc.Width;
-		Desc.height = (uint32)ResourceDesc.Height;
-		Desc.size = Desc.width * Desc.height * 4;
-		Desc.format = DVP_BGRA;
-		Desc.type = DVP_UNSIGNED_BYTE;
+		Desc.width = (uint32) InArgs.Width;
+		Desc.height = (uint32) InArgs.Height;
+		Desc.size = InArgs.Stride * Desc.height;
+		Desc.format = InArgs.PixelFormat == EPixelFormat::PF_8Bit ? DVP_BGRA : DVP_RGBA_INTEGER;
+		Desc.type = InArgs.PixelFormat == EPixelFormat::PF_8Bit ? DVP_UNSIGNED_BYTE : DVP_INT;
 		Desc.handleType = DVP_OPAQUE_WIN32;
 		Desc.external.handle = OutTextureInfo->External.Handle;
 
@@ -114,3 +111,5 @@ namespace UE::GPUTextureTransfer::Private
 		return (dvpMapBufferEndD3D12(InHandle, D3DCommandQueue));
 	}
 }
+
+#endif // PLATFORM_WINDOWS
diff --git a/Engine/Plugins/Media/MediaIOFramework/Source/GPUTextureTransfer/Private/D3D12TextureTransfer.h b/Engine/Plugins/Media/MediaIOFramework/Source/GPUTextureTransfer/Private/D3D12TextureTransfer.h
--- a/Engine/Plugins/Media/MediaIOFramework/Source/GPUTextureTransfer/Private/D3D12TextureTransfer.h
+++ b/Engine/Plugins/Media/MediaIOFramework/Source/GPUTextureTransfer/Private/D3D12TextureTransfer.h
@@ -2,8 +2,10 @@
 
 #pragma once
 
+#if PLATFORM_WINDOWS
 #include "TextureTransferBase.h"
 
+
 struct ID3D12CommandQueue;
 struct ID3D12Device;
 
@@ -15,7 +17,7 @@ protected:
 	virtual DVPStatus Init_Impl(const FInitializeDMAArgs& InArgs) override;
 	virtual DVPStatus GetConstants_Impl(uint32* OutBufferAddrAlignment, uint32* OutBufferGPUStrideAlignment, uint32* OutSemaphoreAddrAlignment, uint32* OutSemaphoreAllocSize, uint32* OutSemaphorePayloadOffset, uint32* OutSemaphorePayloadSize) const override;
 	virtual DVPStatus BindBuffer_Impl(DVPBufferHandle InBufferHandle) const override;
-	virtual DVPStatus CreateGPUResource_Impl(void* InTexture, FTextureInfo* OutTextureInfo) const override;
+	virtual DVPStatus CreateGPUResource_Impl(const FRegisterDMATextureArgs& InArgs, FTextureInfo* OutTextureInfo) const override;
 	virtual DVPStatus UnbindBuffer_Impl(DVPBufferHandle InBufferHandle) const override;
 	virtual DVPStatus CloseDevice_Impl() const override;
 	virtual DVPStatus MapBufferWaitAPI_Impl(DVPBufferHandle InHandle) const;
@@ -26,4 +28,6 @@ private:
 	ID3D12Device* D3DDevice = nullptr;
 	ID3D12CommandQueue* D3DCommandQueue = nullptr;
 };
-}
\ No newline at end of file
+}
+
+#endif // PLATFORM_WINDOWS
\ No newline at end of file
diff --git a/Engine/Plugins/Media/MediaIOFramework/Source/GPUTextureTransfer/Private/GPUTextureTransferModule.cpp b/Engine/Plugins/Media/MediaIOFramework/Source/GPUTextureTransfer/Private/GPUTextureTransferModule.cpp
--- a/Engine/Plugins/Media/MediaIOFramework/Source/GPUTextureTransfer/Private/GPUTextureTransferModule.cpp
+++ b/Engine/Plugins/Media/MediaIOFramework/Source/GPUTextureTransfer/Private/GPUTextureTransferModule.cpp
@@ -2,18 +2,20 @@
 
 #include "GPUTextureTransferModule.h"
 
+#if PLATFORM_WINDOWS
 #include "D3D11TextureTransfer.h"
 #include "D3D12TextureTransfer.h"
 #include "VulkanTextureTransfer.h"
-#include "RHI.h"
-#include "DynamicRHI.h"
+#include "TextureTransferBase.h"
+#endif
+
+#include "CoreMinimal.h"
 #include "IVulkanDynamicRHI.h"
 #include "GenericPlatform/GenericPlatformDriver.h"
 #include "HAL/PlatformMisc.h"
 #include "Misc/CoreDelegates.h"
 #include "Modules/ModuleManager.h"
 #include "RenderingThread.h"
-#include "TextureTransferBase.h"
 
 DEFINE_LOG_CATEGORY(LogGPUTextureTransfer);
 
@@ -35,22 +37,32 @@ void FGPUTextureTransferModule::StartupModule()
 {
 	if (LoadGPUDirectBinary())
 	{
-		bIsGPUTextureTransferAvailable = true;
-
-		const FGPUDriverInfo GPUDriverInfo = FPlatformMisc::GetGPUDriverInfo(GRHIAdapterName);
-		bIsGPUTextureTransferAvailable = GPUDriverInfo.IsNVIDIA() && !FModuleManager::Get().IsModuleLoaded("RenderDocPlugin");
-		bIsGPUTextureTransferAvailable &= !GPUDriverInfo.DeviceDescription.Contains(TEXT("Tesla"));
-
-		TransferObjects.AddDefaulted(RHI_MAX);
+#if PLATFORM_WINDOWS		
+		const TCHAR* DynamicRHIModuleName = GetSelectedDynamicRHIModuleName(false);
+#elif PLATFORM_LINUX
+		const TCHAR* DynamicRHIModuleName = TEXT("VulkanRHI");
+#else
+		const TCHAR* DynamicRHIModuleName = TEXT("");
+		ensure(false);
+#endif
 
-		if (bIsGPUTextureTransferAvailable)
+		// We cannot use GDynmicRHI here because it hasn't been assigned yet.
+		if (TEXT("VulkanRHI") == FString(DynamicRHIModuleName))
 		{
-			// todo: Remove this after testing packaged, probably not needed.
-			FCoreDelegates::OnAllModuleLoadingPhasesComplete.AddRaw(this, &FGPUTextureTransferModule::InitializeTextureTransfer);
-			//Same for shutdown, uninitialize ourselves before library is unloaded
-			FCoreDelegates::OnEnginePreExit.AddRaw(this, &FGPUTextureTransferModule::UninitializeTextureTransfer);
+#if PLATFORM_WINDOWS
+			const TArray<const ANSICHAR*> ExtentionsToAdd{ VK_KHR_EXTERNAL_MEMORY_EXTENSION_NAME, VK_KHR_EXTERNAL_MEMORY_WIN32_EXTENSION_NAME, VK_KHR_EXTERNAL_SEMAPHORE_WIN32_EXTENSION_NAME, VK_KHR_SURFACE_EXTENSION_NAME, VK_KHR_WIN32_SURFACE_EXTENSION_NAME };
+#elif PLATFORM_LINUX
+			const TArray<const ANSICHAR*> ExtentionsToAdd{ VK_KHR_EXTERNAL_MEMORY_EXTENSION_NAME, VK_KHR_EXTERNAL_MEMORY_FD_EXTENSION_NAME, VK_KHR_EXTERNAL_SEMAPHORE_FD_EXTENSION_NAME, VK_KHR_SURFACE_EXTENSION_NAME };
+#endif
+			IVulkanDynamicRHI::AddEnabledDeviceExtensionsAndLayers(ExtentionsToAdd, TArray<const ANSICHAR*>());
 		}
 
+		TransferObjects.AddDefaulted(RHI_MAX);
+
+		// Since this module is started before the RHI is initialized, we have to delay initialization to later. 
+		FCoreDelegates::OnAllModuleLoadingPhasesComplete.AddRaw(this, &FGPUTextureTransferModule::InitializeTextureTransfer);
+		// Same for shutdown, uninitialize ourselves before library is unloaded
+		FCoreDelegates::OnEnginePreExit.AddRaw(this, &FGPUTextureTransferModule::UninitializeTextureTransfer);
 	}
 }
 
@@ -58,8 +70,9 @@ void FGPUTextureTransferModule::ShutdownModule()
 {
 }
 
-FGPUTextureTransferModule::TextureTransferPtr FGPUTextureTransferModule::GetTextureTransfer()
+UE::GPUTextureTransfer::TextureTransferPtr FGPUTextureTransferModule::GetTextureTransfer()
 {
+#if PLATFORM_WINDOWS
 	UE::GPUTextureTransfer::ERHI SupportedRHI = ConvertRHI(RHIGetInterfaceType());
 	if (SupportedRHI == UE::GPUTextureTransfer::ERHI::Invalid) 
 	{
@@ -72,13 +85,17 @@ FGPUTextureTransferModule::TextureTransferPtr FGPUTextureTransferModule::GetText
 	{
 		return TransferObjects[RHIIndex];
 	}
-
+#endif
 	return nullptr;
 }
 
 bool FGPUTextureTransferModule::IsAvailable()
 {
+#if PLATFORM_WINDOWS
 	return FModuleManager::Get().IsModuleLoaded("GPUTextureTransfer");
+#else
+	return false;
+#endif
 }
 
 FGPUTextureTransferModule& FGPUTextureTransferModule::Get()
@@ -88,14 +105,13 @@ FGPUTextureTransferModule& FGPUTextureTransferModule::Get()
 
 bool FGPUTextureTransferModule::LoadGPUDirectBinary()
 {
+#if PLATFORM_WINDOWS
 	FString GPUDirectPath = FPaths::Combine(FPaths::EngineDir(), TEXT("Binaries/ThirdParty/NVIDIA/GPUDirect"), FPlatformProcess::GetBinariesSubdirectory());
 	FPlatformProcess::PushDllDirectory(*GPUDirectPath);
 
 	FString DVPDll;
 
-#if PLATFORM_WINDOWS
 	DVPDll = TEXT("dvp.dll");
-#endif
 
 	DVPDll = FPaths::Combine(GPUDirectPath, DVPDll);
 
@@ -107,12 +123,25 @@ bool FGPUTextureTransferModule::LoadGPUDirectBinary()
 
 	FPlatformProcess::PopDllDirectory(*GPUDirectPath);
 
+#endif
 	return !!TextureTransferHandle;
 }
 
 
 void FGPUTextureTransferModule::InitializeTextureTransfer()
 {
+#if PLATFORM_WINDOWS
+	bIsGPUTextureTransferAvailable = true;
+	// This must be called on game thread 
+	const FGPUDriverInfo GPUDriverInfo = FPlatformMisc::GetGPUDriverInfo(GRHIAdapterName);
+	bIsGPUTextureTransferAvailable = GPUDriverInfo.IsNVIDIA() && !FModuleManager::Get().IsModuleLoaded("RenderDocPlugin");
+	bIsGPUTextureTransferAvailable = bIsGPUTextureTransferAvailable && !GPUDriverInfo.DeviceDescription.Contains(TEXT("Tesla"));
+
+	if (!bIsGPUTextureTransferAvailable)
+	{
+		return;
+	}
+
 	ENQUEUE_RENDER_COMMAND(InitializeGPUTextureTransfer)(
 	[this](FRHICommandListImmediate& RHICmdList) mutable
 	{
@@ -121,7 +150,7 @@ void FGPUTextureTransferModule::InitializeTextureTransfer()
 			return;
 		}
 
-		TextureTransferPtr TextureTransfer;
+		UE::GPUTextureTransfer::TextureTransferPtr TextureTransfer;
 
 		UE::GPUTextureTransfer::ERHI RHI = ConvertRHI(RHIGetInterfaceType());
 
@@ -158,21 +187,24 @@ void FGPUTextureTransferModule::InitializeTextureTransfer()
 			TransferObjects[RHIIndex] = TextureTransfer;
 		}
 	});
+#endif // PLATFORM_WINDOWS
 }
 
 void FGPUTextureTransferModule::UninitializeTextureTransfer()
 {
+#if PLATFORM_WINDOWS
 	ENQUEUE_RENDER_COMMAND(UninitializeGPUTextureTransfer)(
 		[this](FRHICommandListImmediate& RHICmdList) mutable
 		{
 			for (uint8 RhiIt = 1; RhiIt < RHI_MAX; RhiIt++)
 			{
-				if (const TextureTransferPtr& TextureTransfer = TransferObjects[RhiIt])
+				if (const UE::GPUTextureTransfer::TextureTransferPtr& TextureTransfer = TransferObjects[RhiIt])
 				{
 					TextureTransfer->Uninitialize();
 				}
 			}
 		});
+#endif
 }
 
 IMPLEMENT_MODULE(FGPUTextureTransferModule, GPUTextureTransfer);
diff --git a/Engine/Plugins/Media/MediaIOFramework/Source/GPUTextureTransfer/Private/TextureTransferBase.cpp b/Engine/Plugins/Media/MediaIOFramework/Source/GPUTextureTransfer/Private/TextureTransferBase.cpp
--- a/Engine/Plugins/Media/MediaIOFramework/Source/GPUTextureTransfer/Private/TextureTransferBase.cpp
+++ b/Engine/Plugins/Media/MediaIOFramework/Source/GPUTextureTransfer/Private/TextureTransferBase.cpp
@@ -4,16 +4,17 @@
 #include "GPUTextureTransferModule.h"
 #include "Misc/ScopeLock.h"
 
-#if PERF_LOGGING
-#include "ProfilingDebugging/ScopedTimers.h"
-#endif
-
 #if PLATFORM_WINDOWS
 #include "Windows/AllowWindowsPlatformTypes.h"
 #include "Windows/MinWindows.h"
 #include "Windows/HideWindowsPlatformTypes.h"
+
+
+#if PERF_LOGGING
+#include "ProfilingDebugging/ScopedTimers.h"
 #endif
 
+
 #if PERF_LOGGING
 #define LOG_PERF(FuncName)\
 FAutoScopedDurationTimer AutoTimer;\
@@ -150,52 +151,62 @@ namespace UE::GPUTextureTransfer::Private
 		}
 	}
 
-	bool FTextureTransferBase::TransferTexture(void* InBuffer, void* InRHITexture, ETransferDirection TransferDirection)
+	bool FTextureTransferBase::TransferTexture(void* InBuffer, FRHITexture* InRHITexture, ETransferDirection TransferDirection)
 	{
 		FScopeLock ScopeLock{ &CriticalSection };
 
 		FExternalBufferInfo* BufferInfo = RegisteredBuffers.Find(InBuffer);
 		FTextureInfo* TextureInfo = RegisteredTextures.Find(InRHITexture);
 
-		if (BufferInfo && TextureInfo)
+		if (!BufferInfo)
 		{
-			if (BufferInfo->GPUMemorySync && BufferInfo->SystemMemorySync)
-			{
-				BufferInfo->GPUMemorySync->ReleaseValue++;
+			UE_LOG(LogGPUTextureTransfer, Error, TEXT("Error while performing a GPU transfer texture, CPU Buffer %u was not registered."), reinterpret_cast<std::uintptr_t>(InBuffer));
+			return false;
+		}
 
-				DVP_CALL(dvpBegin());
-				DVP_CALL(dvpMapBufferWaitDVP(TextureInfo->DVPHandle));
-			
-				DVPStatus Status = DVP_STATUS_OK;
-				if (TransferDirection == ETransferDirection::GPU_TO_CPU)
-				{
-					Status = dvpMemcpy2D(TextureInfo->DVPHandle, BufferInfo->SystemMemorySync->DVPSyncObject,
-						BufferInfo->SystemMemorySync->AcquireValue, DVP_TIMEOUT_IGNORED, BufferInfo->DVPHandle,
-						BufferInfo->GPUMemorySync->DVPSyncObject, BufferInfo->GPUMemorySync->ReleaseValue, 0, 0, BufferInfo->Height, BufferInfo->Width);
-				}
-				else
-				{
-					Status = dvpMemcpy2D(BufferInfo->DVPHandle, BufferInfo->SystemMemorySync->DVPSyncObject,
-						BufferInfo->SystemMemorySync->AcquireValue, DVP_TIMEOUT_IGNORED, TextureInfo->DVPHandle,
-						BufferInfo->GPUMemorySync->DVPSyncObject, BufferInfo->GPUMemorySync->ReleaseValue, 0, 0, BufferInfo->Height, BufferInfo->Width);
-				}
+		if (!TextureInfo)
+		{
+			UE_LOG(LogGPUTextureTransfer, Error, TEXT("Error while performing a GPU transfer texture, texture %u was not registered."), reinterpret_cast<std::uintptr_t>(InRHITexture));
+			return false;
+		}
 
-				BufferInfo->SystemMemorySync->AcquireValue++;
-				if (Status != DVP_STATUS_OK)
-				{
-					UE_LOG(LogGPUTextureTransfer, Error, TEXT("Error while performing a GPU transfer texture. Error: '%d'."), Status);
-					return false;
-				}
+		if (BufferInfo->GPUMemorySync && BufferInfo->SystemMemorySync)
+		{
+			BufferInfo->GPUMemorySync->ReleaseValue++;
 
-				DVP_CALL(dvpMapBufferEndDVP(TextureInfo->DVPHandle));
-				DVP_CALL(dvpEnd())
+			DVP_CALL(dvpBegin());
+			DVP_CALL(dvpMapBufferWaitDVP(TextureInfo->DVPHandle));
+			
+			DVPStatus Status = DVP_STATUS_OK;
+			if (TransferDirection == ETransferDirection::GPU_TO_CPU)
+			{
+				Status = dvpMemcpy2D(TextureInfo->DVPHandle, BufferInfo->SystemMemorySync->DVPSyncObject,
+					BufferInfo->SystemMemorySync->AcquireValue, DVP_TIMEOUT_IGNORED, BufferInfo->DVPHandle,
+					BufferInfo->GPUMemorySync->DVPSyncObject, BufferInfo->GPUMemorySync->ReleaseValue, 0, 0, BufferInfo->Height, BufferInfo->Width);
 			}
+			else
+			{
+				Status = dvpMemcpy2D(BufferInfo->DVPHandle, BufferInfo->SystemMemorySync->DVPSyncObject,
+					BufferInfo->SystemMemorySync->AcquireValue, DVP_TIMEOUT_IGNORED, TextureInfo->DVPHandle,
+					BufferInfo->GPUMemorySync->DVPSyncObject, BufferInfo->GPUMemorySync->ReleaseValue, 0, 0, BufferInfo->Height, BufferInfo->Width);
+			}
+
+			BufferInfo->SystemMemorySync->AcquireValue++;
+
+			DVP_CALL(dvpMapBufferEndDVP(TextureInfo->DVPHandle));
+			DVP_CALL(dvpEnd())
 
+			if (Status != DVP_STATUS_OK)
+			{
+				UE_LOG(LogGPUTextureTransfer, Error, TEXT("Error while performing a GPU transfer texture. Error: '%d'."), Status);
+				return false;
+
+			}
 			return true;
 		}
 		else
 		{
-			UE_LOG(LogGPUTextureTransfer, Error, TEXT("Could not find both the bufferinfo and the texture info for a GPU DMA transfer."));
+			UE_LOG(LogGPUTextureTransfer, Warning, TEXT("Error while performing a GPU transfer texture: A sync object was not found."));
 		}
 
 		return false;
@@ -220,14 +231,16 @@ namespace UE::GPUTextureTransfer::Private
 			BufferInfo.SystemMemorySync = MakeUnique<DVPSync>(SemaphoreAllocSize, SemaphoreAddressAlignment);
 			BufferInfo.GPUMemorySync = MakeUnique<DVPSync>(SemaphoreAllocSize, SemaphoreAddressAlignment);
 
+			// todo jroy: Add conversion method from unreal pixel format to dvp pixel format to support more format/type combinations.
+
 			// Register system memory buffers with DVP
 			DVPSysmemBufferDesc SystemMemoryBuffersDesc;
 			SystemMemoryBuffersDesc.width = Args.Width;
 			SystemMemoryBuffersDesc.height = Args.Height;
 			SystemMemoryBuffersDesc.stride = Args.Stride;
 			SystemMemoryBuffersDesc.size = 0; // Only needed with DVP_BUFFER
-			SystemMemoryBuffersDesc.format = DVP_BGRA;
-			SystemMemoryBuffersDesc.type = DVP_UNSIGNED_BYTE;
+			SystemMemoryBuffersDesc.format = Args.PixelFormat == EPixelFormat::PF_8Bit ? DVP_BGRA : DVP_RGBA_INTEGER;
+			SystemMemoryBuffersDesc.type = Args.PixelFormat == EPixelFormat::PF_8Bit ? DVP_UNSIGNED_BYTE : DVP_INT;
 			SystemMemoryBuffersDesc.bufAddr = Args.Buffer;
 
 			DVP_CALL(dvpCreateBuffer(&SystemMemoryBuffersDesc, &BufferInfo.DVPHandle));
@@ -255,12 +268,12 @@ namespace UE::GPUTextureTransfer::Private
 		if (!RegisteredTextures.Contains(Args.RHITexture))
 		{
 			FTextureInfo Info;
-			DVP_CALL(CreateGPUResource_Impl(Args.RHITexture, &Info));
+			DVP_CALL(CreateGPUResource_Impl(Args, &Info));
 			RegisteredTextures.Add({ Args.RHITexture, std::move(Info) });
 		}
 	}
 
-	void FTextureTransferBase::UnregisterTexture(void* InRHITexture)
+	void FTextureTransferBase::UnregisterTexture(FRHITexture* InRHITexture)
 	{
 		FScopeLock ScopeLock{ &CriticalSection };
 		if (FTextureInfo* TextureInfo = RegisteredTextures.Find(InRHITexture))
@@ -269,7 +282,7 @@ namespace UE::GPUTextureTransfer::Private
 			if (TextureInfo->External.Handle)
 			{
 #if PLATFORM_WINDOWS
-				CloseHandle(TextureInfo->External.Handle);
+				::CloseHandle(TextureInfo->External.Handle);
 #endif
 			}
 			RegisteredTextures.Remove(InRHITexture);
@@ -286,7 +299,7 @@ namespace UE::GPUTextureTransfer::Private
 				if (It.Value().External.Handle)
 				{
 #if PLATFORM_WINDOWS
-					CloseHandle(It.Value().External.Handle);
+					::CloseHandle(It.Value().External.Handle);
 #endif
 				}
 			}
@@ -307,7 +320,7 @@ namespace UE::GPUTextureTransfer::Private
 		RegisteredBuffers.Reset();
 	}
 
-	void FTextureTransferBase::LockTexture(void* InRHITexture)
+	void FTextureTransferBase::LockTexture(FRHITexture* InRHITexture)
 	{
 		FScopeLock ScopeLock{ &CriticalSection };
 		if (FTextureInfo* TextureInfo = RegisteredTextures.Find(InRHITexture))
@@ -316,7 +329,7 @@ namespace UE::GPUTextureTransfer::Private
 		}
 	}
 
-	void FTextureTransferBase::UnlockTexture(void* InRHITexture)
+	void FTextureTransferBase::UnlockTexture(FRHITexture* InRHITexture)
 	{
 		LOG_PERF(UnlockTexture);
 		FScopeLock ScopeLock{ &CriticalSection };
@@ -326,6 +339,16 @@ namespace UE::GPUTextureTransfer::Private
 		}
 	}
 
+	uint32 FTextureTransferBase::GetBufferAlignment() const
+	{
+		return BufferAddressAlignment;
+	}
+
+	uint32 FTextureTransferBase::GetTextureStride() const
+	{
+		return BufferGpuStrideAlignment;
+	}
+
 	void FTextureTransferBase::ClearBufferInfo(FExternalBufferInfo& BufferInfo)
 	{
 		DVP_CALL(UnbindBuffer_Impl(BufferInfo.DVPHandle));
@@ -345,3 +368,4 @@ namespace UE::GPUTextureTransfer::Private
 	}
 
 }
+#endif // PLATFORM_WINDOWS
diff --git a/Engine/Plugins/Media/MediaIOFramework/Source/GPUTextureTransfer/Private/TextureTransferBase.h b/Engine/Plugins/Media/MediaIOFramework/Source/GPUTextureTransfer/Private/TextureTransferBase.h
--- a/Engine/Plugins/Media/MediaIOFramework/Source/GPUTextureTransfer/Private/TextureTransferBase.h
+++ b/Engine/Plugins/Media/MediaIOFramework/Source/GPUTextureTransfer/Private/TextureTransferBase.h
@@ -7,14 +7,11 @@
 #if PLATFORM_WINDOWS
 #include "Windows/AllowWindowsPlatformTypes.h"
 #include "Windows/PreWindowsApi.h"
-#endif
-
 #include "DVPAPI.h"
-
-#if PLATFORM_WINDOWS
 #include "Windows/PostWindowsApi.h"
 #include "Windows/HideWindowsPlatformTypes.h"
-#endif
+
+class FRHITexture;
 
 namespace UE::GPUTextureTransfer::Private
 {
@@ -26,13 +23,15 @@ namespace UE::GPUTextureTransfer::Private
 		//~ Begin ITextureTransfer interface
 		virtual bool BeginSync(void* InBuffer, ETransferDirection TransferDirection) override;
 		virtual void EndSync(void* InBuffer) override;
-		virtual bool TransferTexture(void* InBuffer, void* InRHITexture, ETransferDirection TransferDirection) override;
+		virtual bool TransferTexture(void* InBuffer, FRHITexture* InRHITexture, ETransferDirection TransferDirection) override;
 		virtual void RegisterBuffer(const FRegisterDMABufferArgs& Args) override;
 		virtual void UnregisterBuffer(void* InBuffer) override;
 		virtual void RegisterTexture(const FRegisterDMATextureArgs& Args) override;
-		virtual void UnregisterTexture(void* InRHITexture) override;
-		virtual void LockTexture(void* InRHITexture) override;
-		virtual void UnlockTexture(void* InRHITexture) override;
+		virtual void UnregisterTexture(FRHITexture* InRHITexture) override;
+		virtual uint32 GetBufferAlignment() const override;
+		virtual uint32 GetTextureStride() const override;
+		virtual void LockTexture(FRHITexture* InRHITexture) override;
+		virtual void UnlockTexture(FRHITexture* InRHITexture) override;
 		// End ITextureTransfer interface
 
 	protected:
@@ -65,7 +64,8 @@ namespace UE::GPUTextureTransfer::Private
 		struct FTextureInfo
 		{
 			DVPBufferHandle DVPHandle = 0;
-			union 
+
+			union
 			{
 				void* Handle = nullptr;
 				int Fd;
@@ -76,13 +76,12 @@ namespace UE::GPUTextureTransfer::Private
 		virtual DVPStatus Init_Impl(const FInitializeDMAArgs& InArgs) = 0;
 		virtual DVPStatus GetConstants_Impl(uint32* OutBufferAddrAlignment, uint32* OutBufferGPUStrideAlignment, uint32* OutSemaphoreAddrAlignment, uint32* OutSemaphoreAllocSize, uint32* OutSemaphorePayloadOffset, uint32* OutSemaphorePayloadSize) const = 0;
 		virtual DVPStatus BindBuffer_Impl(DVPBufferHandle InBufferHandle) const = 0;
-		virtual DVPStatus CreateGPUResource_Impl(void* InTexture, FTextureInfo* OutTextureInfo) const = 0;
+		virtual DVPStatus CreateGPUResource_Impl(const FRegisterDMATextureArgs& InArgs, FTextureInfo* OutTextureInfo) const = 0;
 		virtual DVPStatus UnbindBuffer_Impl(DVPBufferHandle InBufferHandle) const = 0;
 		virtual DVPStatus CloseDevice_Impl() const = 0;
 		virtual DVPStatus MapBufferWaitAPI_Impl(DVPBufferHandle InHandle) const;
 		virtual DVPStatus MapBufferEndAPI_Impl(DVPBufferHandle InHandle) const;
 
-		
 	private:
 		void ClearBufferInfo(FExternalBufferInfo& BufferInfo);
 		void ClearRegisteredTextures();
@@ -91,7 +90,7 @@ namespace UE::GPUTextureTransfer::Private
 
 		bool bInitialized = false;
 		TMap<void*, FExternalBufferInfo> RegisteredBuffers;
-		TMap<void*, FTextureInfo> RegisteredTextures;
+		TMap<FRHITexture*, FTextureInfo> RegisteredTextures;
 
 		uint32 BufferAddressAlignment = 0;
 		uint32 BufferGpuStrideAlignment = 0;
@@ -103,3 +102,4 @@ namespace UE::GPUTextureTransfer::Private
 		FCriticalSection CriticalSection;
 	};
 }
+#endif // PLATFORM_WINDOWS
diff --git a/Engine/Plugins/Media/MediaIOFramework/Source/GPUTextureTransfer/Private/VulkanTextureTransfer.cpp b/Engine/Plugins/Media/MediaIOFramework/Source/GPUTextureTransfer/Private/VulkanTextureTransfer.cpp
--- a/Engine/Plugins/Media/MediaIOFramework/Source/GPUTextureTransfer/Private/VulkanTextureTransfer.cpp
+++ b/Engine/Plugins/Media/MediaIOFramework/Source/GPUTextureTransfer/Private/VulkanTextureTransfer.cpp
@@ -5,15 +5,12 @@
 #if PLATFORM_WINDOWS
 #include "Windows/AllowWindowsPlatformTypes.h"
 #include "Windows/PreWindowsApi.h"
-#endif
 
 #include "DVPAPI.h"
 #include "dvpapi_vulkan.h"
 
-#if PLATFORM_WINDOWS
 #include "Windows/PostWindowsApi.h"
 #include "Windows/HideWindowsPlatformTypes.h"
-#endif
 
 
 namespace UE::GPUTextureTransfer::Private
@@ -56,22 +53,37 @@ namespace UE::GPUTextureTransfer::Private
 		return dvpBindToVkDevice(InBufferHandle, VulkanDevice);
 	}
 
-	DVPStatus FVulkanTextureTransfer::CreateGPUResource_Impl(void* InTexture, FTextureTransferBase::FTextureInfo* OutTextureInfo) const
+	DVPStatus FVulkanTextureTransfer::CreateGPUResource_Impl(const FRegisterDMATextureArgs& InArgs, FTextureTransferBase::FTextureInfo* OutTextureInfo) const
 	{
 		if (!OutTextureInfo || !VulkanDevice)
 		{
 			return DVP_STATUS_ERROR;
 		}
 
-		// todo for Vulkan
-		checkNoEntry();
+		IVulkanDynamicRHI* VulkanRHI = GetIVulkanDynamicRHI();
+		const FVulkanRHIAllocationInfo TextureAllocationInfo = VulkanRHI->RHIGetAllocationInfo(InArgs.RHITexture);
+
+		VkMemoryGetWin32HandleInfoKHR MemoryGetHandleInfoKHR = {};
+		MemoryGetHandleInfoKHR.sType = VK_STRUCTURE_TYPE_MEMORY_GET_WIN32_HANDLE_INFO_KHR;
+		MemoryGetHandleInfoKHR.pNext = NULL;
+		MemoryGetHandleInfoKHR.memory = TextureAllocationInfo.Handle;
+		MemoryGetHandleInfoKHR.handleType = VK_EXTERNAL_MEMORY_HANDLE_TYPE_OPAQUE_WIN32_BIT;
+
+		PFN_vkGetMemoryWin32HandleKHR GetMemoryWin32HandleKHR = (PFN_vkGetMemoryWin32HandleKHR)VulkanRHI->RHIGetVkDeviceProcAddr("vkGetMemoryWin32HandleKHR");
+		if (!GetMemoryWin32HandleKHR)
+		{
+			return DVP_STATUS_ERROR;
+		}
+
+		HANDLE Handle;
+		VERIFYVULKANRESULT_EXTERNAL(GetMemoryWin32HandleKHR(VulkanDevice, &MemoryGetHandleInfoKHR, &Handle));//&OutTextureInfo->External.Handle));
 
 		DVPGpuExternalResourceDesc Desc;
-		//Desc.width = (uint32) ResourceDesc.Width;
-		//Desc.height = (uint32)ResourceDesc.Height;
-		//Desc.size = Desc.width * Desc.height * 4;
-		Desc.format = DVP_BGRA;
-		Desc.type = DVP_UNSIGNED_BYTE;
+		Desc.width = InArgs.Width;
+		Desc.height = InArgs.Height;
+		Desc.size = InArgs.Stride * Desc.height;
+		Desc.format = InArgs.PixelFormat == EPixelFormat::PF_8Bit ? DVP_BGRA : DVP_RGBA_INTEGER;
+		Desc.type = InArgs.PixelFormat == EPixelFormat::PF_8Bit ? DVP_UNSIGNED_BYTE : DVP_INT;
 		Desc.handleType = DVP_OPAQUE_WIN32;
 		Desc.external.handle = OutTextureInfo->External.Handle;
 
@@ -108,3 +120,5 @@ namespace UE::GPUTextureTransfer::Private
 		return (dvpMapBufferEndVk(InHandle, VulkanQueue));
 	}
 }
+
+#endif // PLATFORM_WINDOWS
\ No newline at end of file
diff --git a/Engine/Plugins/Media/MediaIOFramework/Source/GPUTextureTransfer/Private/VulkanTextureTransfer.h b/Engine/Plugins/Media/MediaIOFramework/Source/GPUTextureTransfer/Private/VulkanTextureTransfer.h
--- a/Engine/Plugins/Media/MediaIOFramework/Source/GPUTextureTransfer/Private/VulkanTextureTransfer.h
+++ b/Engine/Plugins/Media/MediaIOFramework/Source/GPUTextureTransfer/Private/VulkanTextureTransfer.h
@@ -2,6 +2,8 @@
 
 #pragma once
 
+#if PLATFORM_WINDOWS
+
 #include "TextureTransferBase.h"
 #include "IVulkanDynamicRHI.h"
 
@@ -13,7 +15,7 @@ protected:
 	virtual DVPStatus Init_Impl(const FInitializeDMAArgs& InArgs) override;
 	virtual DVPStatus GetConstants_Impl(uint32* OutBufferAddrAlignment, uint32* OutBufferGPUStrideAlignment, uint32* OutSemaphoreAddrAlignment, uint32* OutSemaphoreAllocSize, uint32* OutSemaphorePayloadOffset, uint32* OutSemaphorePayloadSize) const override;
 	virtual DVPStatus BindBuffer_Impl(DVPBufferHandle InBufferHandle) const override;
-	virtual DVPStatus CreateGPUResource_Impl(void* InTexture, FTextureInfo* OutTextureInfo) const override;
+	virtual DVPStatus CreateGPUResource_Impl(const FRegisterDMATextureArgs& InArgs, FTextureInfo* OutTextureInfo) const override;
 	virtual DVPStatus UnbindBuffer_Impl(DVPBufferHandle InBufferHandle) const override;
 	virtual DVPStatus CloseDevice_Impl() const override;
 	virtual DVPStatus MapBufferWaitAPI_Impl(DVPBufferHandle InHandle) const;
@@ -25,4 +27,6 @@ private:
 	VkDevice VulkanDevice = nullptr;
 	VkQueue VulkanQueue = nullptr;
 };
-}
\ No newline at end of file
+}
+
+#endif
\ No newline at end of file
diff --git a/Engine/Plugins/Media/MediaIOFramework/Source/GPUTextureTransfer/Public/GPUTextureTransfer.h b/Engine/Plugins/Media/MediaIOFramework/Source/GPUTextureTransfer/Public/GPUTextureTransfer.h
--- a/Engine/Plugins/Media/MediaIOFramework/Source/GPUTextureTransfer/Public/GPUTextureTransfer.h
+++ b/Engine/Plugins/Media/MediaIOFramework/Source/GPUTextureTransfer/Public/GPUTextureTransfer.h
@@ -3,6 +3,8 @@
 #pragma once
 
 #include "CoreMinimal.h"
+#include "RHI.h"
+#include "RHIResources.h"
 
 namespace UE::GPUTextureTransfer
 {
@@ -31,18 +33,32 @@ namespace UE::GPUTextureTransfer
 		// End Vulkan Only
 	};
 
+	enum class EPixelFormat : uint8_t
+	{
+		PF_8Bit,
+		PF_10Bit
+	};
+
 	struct GPUTEXTURETRANSFER_API FRegisterDMABufferArgs
 	{
 		void* Buffer = nullptr;
 		uint32 Width = 0;
 		uint32 Height = 0;
 		uint32 Stride = 0;
+		EPixelFormat PixelFormat = EPixelFormat::PF_8Bit;
 	};
 
 	struct GPUTEXTURETRANSFER_API FRegisterDMATextureArgs
 	{
-		void* RHITexture = nullptr;
+		FRHITexture* RHITexture = nullptr;
 		void* RHIResourceMemory = nullptr; // Vulkan only
+		uint32 Width = 0;
+		uint32 Height = 0;
+		EPixelFormat PixelFormat = EPixelFormat::PF_8Bit;
+		// Stride in bytes
+		uint32 Stride = 0;
+		// Only for VK
+		void* SharedHandle = nullptr;
 	};
 
 	enum class ETransferDirection :uint8
@@ -60,17 +76,17 @@ namespace UE::GPUTextureTransfer
 		virtual void UnregisterBuffer(void* InBuffer) = 0;
 
 		virtual void RegisterTexture(const FRegisterDMATextureArgs& Args) = 0;
-		virtual void UnregisterTexture(void* RHITexture) = 0;
+		virtual void UnregisterTexture(FRHITexture* RHITexture) = 0;
 
 		/**
 		 * Calling this will prevent the DVP library from using the RHI texture passed as argument until Unlock is called.
 		 */
-		virtual void LockTexture(void* RHITexture) = 0;
+		virtual void LockTexture(FRHITexture* RHITexture) = 0;
 
 		/**
 		 * Calling this will allow the DVP library to access the RHI texture passed as argument.
 		 */
-		virtual void UnlockTexture(void* RHITexture) = 0;
+		virtual void UnlockTexture(FRHITexture* RHITexture) = 0;
 
 		virtual bool BeginSync(void* InBuffer, ETransferDirection TransferDirection) = 0;
 		virtual void EndSync(void* InBuffer) = 0;
@@ -86,7 +102,17 @@ namespace UE::GPUTextureTransfer
 		 *  EndSync(Buffer);
 		 * 
 		 */
-		virtual bool TransferTexture(void* InBuffer, void* RHITexture, ETransferDirection TransferDirection) = 0;
+		virtual bool TransferTexture(void* InBuffer, FRHITexture* RHITexture, ETransferDirection TransferDirection) = 0;
+
+		/**
+		 * Get the recommended alignment for the cpu buffer.
+		 **/
+		virtual uint32 GetBufferAlignment() const = 0;
+
+		/**
+		 * Get the recommended stride for textures that will be  copied.
+		 */
+		virtual uint32 GetTextureStride() const = 0;
 
 		virtual bool Initialize(const FInitializeDMAArgs& Args) = 0;
 		virtual bool Uninitialize() = 0;
diff --git a/Engine/Plugins/Media/MediaIOFramework/Source/GPUTextureTransfer/Public/GPUTextureTransferModule.h b/Engine/Plugins/Media/MediaIOFramework/Source/GPUTextureTransfer/Public/GPUTextureTransferModule.h
--- a/Engine/Plugins/Media/MediaIOFramework/Source/GPUTextureTransfer/Public/GPUTextureTransferModule.h
+++ b/Engine/Plugins/Media/MediaIOFramework/Source/GPUTextureTransfer/Public/GPUTextureTransferModule.h
@@ -7,18 +7,22 @@
 #include "Logging/LogMacros.h"
 #include "Modules/ModuleInterface.h"
 
+#include <atomic>
+
 DECLARE_LOG_CATEGORY_EXTERN(LogGPUTextureTransfer, Log, All);
 
+namespace UE::GPUTextureTransfer
+{
+	using TextureTransferPtr = TSharedPtr<ITextureTransfer>;
+}
+
 class GPUTEXTURETRANSFER_API FGPUTextureTransferModule : public IModuleInterface
 {
-public:
-	using TextureTransferPtr = TSharedPtr<UE::GPUTextureTransfer::ITextureTransfer>;
-	
 public:
 	virtual void StartupModule() override;
 	virtual void ShutdownModule() override;
 
-	TextureTransferPtr GetTextureTransfer();
+	UE::GPUTextureTransfer::TextureTransferPtr GetTextureTransfer();
 
 	static bool IsAvailable();
 	static FGPUTextureTransferModule& Get();
@@ -31,7 +35,7 @@ private:
 private:
 	static constexpr uint8 RHI_MAX = static_cast<uint8>(UE::GPUTextureTransfer::ERHI::RHI_MAX);
 	void* TextureTransferHandle = nullptr;
-	bool bIsGPUTextureTransferAvailable = false;
+	std::atomic<bool> bIsGPUTextureTransferAvailable = false;
 	
-	TArray<TextureTransferPtr> TransferObjects;
+	TArray<UE::GPUTextureTransfer::TextureTransferPtr> TransferObjects;
 };
diff --git a/Engine/Plugins/Media/MediaIOFramework/Source/MediaIOCore/Public/MediaIOCoreTextureSampleBase.h b/Engine/Plugins/Media/MediaIOFramework/Source/MediaIOCore/Public/MediaIOCoreTextureSampleBase.h
--- a/Engine/Plugins/Media/MediaIOFramework/Source/MediaIOCore/Public/MediaIOCoreTextureSampleBase.h
+++ b/Engine/Plugins/Media/MediaIOFramework/Source/MediaIOCore/Public/MediaIOCoreTextureSampleBase.h
@@ -218,7 +218,7 @@ protected:
 	uint32 Height;
 
 	/** Pointer to raw pixels */
-	TArray<uint8> Buffer;
+	TArray<uint8, TAlignedHeapAllocator<4096>> Buffer;
 
 	/** Wheter the sample is in sRGB space and requires an explicit conversion to linear */
 	bool bIsSRGBInput;
