Autodetect aja format input by default

diff --git a/Engine/Plugins/Media/AjaMedia/Source/AjaMedia/Private/Aja/AjaDeviceProvider.cpp b/Engine/Plugins/Media/AjaMedia/Source/AjaMedia/Private/Aja/AjaDeviceProvider.cpp
--- a/Engine/Plugins/Media/AjaMedia/Source/AjaMedia/Private/Aja/AjaDeviceProvider.cpp
+++ b/Engine/Plugins/Media/AjaMedia/Source/AjaMedia/Private/Aja/AjaDeviceProvider.cpp
@@ -235,6 +235,7 @@ public:
 				}
 
 				InOnAutoDetected.ExecuteIfBound(Configurations);
+				AutoDetectChannelPin->Uninitialize();
 			}
 		});
 	}
diff --git a/Engine/Plugins/Media/AjaMedia/Source/AjaMedia/Private/Assets/AjaMediaSource.cpp b/Engine/Plugins/Media/AjaMedia/Source/AjaMedia/Private/Assets/AjaMediaSource.cpp
--- a/Engine/Plugins/Media/AjaMedia/Source/AjaMedia/Private/Assets/AjaMediaSource.cpp
+++ b/Engine/Plugins/Media/AjaMedia/Source/AjaMedia/Private/Assets/AjaMediaSource.cpp
@@ -185,11 +185,23 @@ FString UAjaMediaSource::GetUrl() const
 bool UAjaMediaSource::Validate() const
 {
 	FString FailureReason;
-	if (!MediaConfiguration.IsValid())
+	if (bAutoDetectInput)
 	{
-		UE_LOG(LogAjaMedia, Warning, TEXT("The MediaConfiguration '%s' is invalid."), *GetName());
-		return false;
+		if (!MediaConfiguration.MediaConnection.IsValid())
+		{
+			UE_LOG(LogAjaMedia, Warning, TEXT("The MediaConfiguration '%s' is invalid."), *GetName());
+			return false;
+		}
+	}
+	else
+	{
+		if (!MediaConfiguration.IsValid())
+		{
+			UE_LOG(LogAjaMedia, Warning, TEXT("The MediaConfiguration '%s' is invalid."), *GetName());
+			return false;
+		}
 	}
+	
 
 	if (!FAja::IsInitialized())
 	{
@@ -205,7 +217,9 @@ bool UAjaMediaSource::Validate() const
 
 	TUniquePtr<AJA::AJADeviceScanner> Scanner = MakeUnique<AJA::AJADeviceScanner>();
 	AJA::AJADeviceScanner::DeviceInfo DeviceInfo;
-	if (!Scanner->GetDeviceInfo(MediaConfiguration.MediaConnection.Device.DeviceIdentifier, DeviceInfo))
+	FMediaIODevice Device = MediaConfiguration.MediaConnection.Device;
+
+	if (!Scanner->GetDeviceInfo(Device.DeviceIdentifier, DeviceInfo))
 	{
 		UE_LOG(LogAjaMedia, Warning, TEXT("The MediaSource '%s' use the device '%s' that doesn't exist on this machine."), *GetName(), *MediaConfiguration.MediaConnection.Device.DeviceName.ToString());
 		return false;
diff --git a/Engine/Plugins/Media/AjaMedia/Source/AjaMedia/Private/Player/AjaMediaPlayer.cpp b/Engine/Plugins/Media/AjaMedia/Source/AjaMedia/Private/Player/AjaMediaPlayer.cpp
--- a/Engine/Plugins/Media/AjaMedia/Source/AjaMedia/Private/Player/AjaMediaPlayer.cpp
+++ b/Engine/Plugins/Media/AjaMedia/Source/AjaMedia/Private/Player/AjaMediaPlayer.cpp
@@ -74,6 +74,7 @@
 	, SupportedSampleTypes(EMediaIOSampleType::None)
 	, bPauseRequested(false)
 {
+	DeviceProvider = MakePimpl<FAjaDeviceProvider>();
 }
 
 
@@ -104,154 +105,18 @@
 	{
 		return false;
 	}
-
-	AJA::AJADeviceOptions DeviceOptions(Options->GetMediaOption(AjaMediaOption::DeviceIndex, (int64)0));
-
-	// Read options
-	AJA::AJAInputOutputChannelOptions AjaOptions(TEXT("MediaPlayer"), Options->GetMediaOption(AjaMediaOption::PortIndex, (int64)0));
-	AjaOptions.CallbackInterface = this;
-	AjaOptions.bOutput = false;
+	
+	if (bAutoDetect)
 	{
-		const EMediaIOTransportType TransportType = (EMediaIOTransportType)(Options->GetMediaOption(AjaMediaOption::TransportType, (int64)EMediaIOTransportType::SingleLink));
-		const EMediaIOQuadLinkTransportType QuadTransportType = (EMediaIOQuadLinkTransportType)(Options->GetMediaOption(AjaMediaOption::QuadTransportType, (int64)EMediaIOQuadLinkTransportType::SquareDivision));
-		switch(TransportType)
-		{
-		case EMediaIOTransportType::SingleLink:
-			AjaOptions.TransportType = AJA::ETransportType::TT_SdiSingle;
-			break;
-		case EMediaIOTransportType::DualLink:
-			AjaOptions.TransportType = AJA::ETransportType::TT_SdiDual;
-			break;
-		case EMediaIOTransportType::QuadLink:
-			AjaOptions.TransportType = QuadTransportType == EMediaIOQuadLinkTransportType::SquareDivision ? AJA::ETransportType::TT_SdiQuadSQ : AJA::ETransportType::TT_SdiQuadTSI;
-			break;
-		case EMediaIOTransportType::HDMI:
-			AjaOptions.TransportType = AJA::ETransportType::TT_Hdmi;
-			break;
-		}
+		AJA::AJADeviceOptions DeviceOptions(Options->GetMediaOption(AjaMediaOption::DeviceIndex, (int64)0));
+		DeviceProvider->AutoDetectConfiguration(FAjaDeviceProvider::FOnConfigurationAutoDetected::CreateRaw(this, &FAjaMediaPlayer::OnAutoDetected, Url, Options));
+		return true;
 	}
+	else
 	{
-		const EMediaIOTimecodeFormat Timecode = (EMediaIOTimecodeFormat)(Options->GetMediaOption(AjaMediaOption::TimecodeFormat, (int64)EMediaIOTimecodeFormat::None));
-		bUseFrameTimecode = Timecode != EMediaIOTimecodeFormat::None;
-		AjaOptions.TimecodeFormat = AJA::ETimecodeFormat::TCF_None;
-		switch (Timecode)
-		{
-		case EMediaIOTimecodeFormat::None:
-			AjaOptions.TimecodeFormat = AJA::ETimecodeFormat::TCF_None;
-			break;
-		case EMediaIOTimecodeFormat::LTC:
-			AjaOptions.TimecodeFormat = AJA::ETimecodeFormat::TCF_LTC;
-			break;
-		case EMediaIOTimecodeFormat::VITC:
-			AjaOptions.TimecodeFormat = AJA::ETimecodeFormat::TCF_VITC1;
-			break;
-		default:
-			break;
-		}
-		bEncodeTimecodeInTexel = Options->GetMediaOption(AjaMediaOption::EncodeTimecodeInTexel, false);
+		AJA::AJAInputOutputChannelOptions AjaOptions(TEXT("MediaPlayer"), Options->GetMediaOption(AjaMediaOption::PortIndex, (int64)0));
+		return Open_Internal(Url, Options, AjaOptions);
 	}
-	{
-		const EAjaMediaAudioChannel AudioChannelOption = (EAjaMediaAudioChannel)(Options->GetMediaOption(AjaMediaOption::AudioChannel, (int64)EAjaMediaAudioChannel::Channel8));
-		AjaOptions.NumberOfAudioChannel = (AudioChannelOption == EAjaMediaAudioChannel::Channel8) ? 8 : 6;
-	}
-	{
-		AjaOptions.VideoFormatIndex = Options->GetMediaOption(AjaMediaOption::AjaVideoFormat, (int64)0);
-		LastVideoFormatIndex = AjaOptions.VideoFormatIndex;
-	}
-	{
-		AjaColorFormat= (EAjaMediaSourceColorFormat)(Options->GetMediaOption(AjaMediaOption::ColorFormat, (int64)EAjaMediaSourceColorFormat::YUV2_8bit));
-		switch(AjaColorFormat)
-		{
-		case EAjaMediaSourceColorFormat::YUV2_8bit:
-			if (AjaOptions.bUseKey)
-			{
-				AjaOptions.PixelFormat = AJA::EPixelFormat::PF_8BIT_ARGB;
-			}
-			else
-			{
-				AjaOptions.PixelFormat = AJA::EPixelFormat::PF_8BIT_YCBCR;
-			}
-			break;
-		case EAjaMediaSourceColorFormat::YUV_10bit:
-			if (AjaOptions.bUseKey)
-			{
-				AjaOptions.PixelFormat = AJA::EPixelFormat::PF_10BIT_RGB;
-			}
-			else
-			{
-				AjaOptions.PixelFormat = AJA::EPixelFormat::PF_10BIT_YCBCR;
-			}
-			break;
-		default:
-			AjaOptions.PixelFormat = AJA::EPixelFormat::PF_8BIT_ARGB;
-			break;
-		}
-
-		bIsSRGBInput = Options->GetMediaOption(AjaMediaOption::SRGBInput, false);
-	}
-	{
-		AjaOptions.bUseAncillary = bUseAncillary = Options->GetMediaOption(AjaMediaOption::CaptureAncillary, false);
-		AjaOptions.bUseAudio = bUseAudio = Options->GetMediaOption(AjaMediaOption::CaptureAudio, false);
-		AjaOptions.bUseVideo = bUseVideo = Options->GetMediaOption(AjaMediaOption::CaptureVideo, true);
-		AjaOptions.bUseAutoCirculating = Options->GetMediaOption(AjaMediaOption::CaptureWithAutoCirculating, true);
-		AjaOptions.bUseKey = false;
-		AjaOptions.bBurnTimecode = false;
-		AjaOptions.BurnTimecodePercentY = 80;
-
-		//Adjust supported sample types based on what's being captured
-		SupportedSampleTypes = AjaOptions.bUseVideo ? EMediaIOSampleType::Video : EMediaIOSampleType::None;
-		SupportedSampleTypes |= AjaOptions.bUseAudio ? EMediaIOSampleType::Audio : EMediaIOSampleType::None;
-		SupportedSampleTypes |= AjaOptions.bUseAncillary ? EMediaIOSampleType::Metadata : EMediaIOSampleType::None;
-		Samples->EnableTimedDataChannels(this, SupportedSampleTypes);
-	}
-
-	bVerifyFrameDropCount = Options->GetMediaOption(AjaMediaOption::LogDropFrame, true);
-	MaxNumAudioFrameBuffer = Options->GetMediaOption(AjaMediaOption::MaxAudioFrameBuffer, (int64)8);
-	MaxNumMetadataFrameBuffer = Options->GetMediaOption(AjaMediaOption::MaxAncillaryFrameBuffer, (int64)8);
-	MaxNumVideoFrameBuffer = Options->GetMediaOption(AjaMediaOption::MaxVideoFrameBuffer, (int64)8);
-
-	check(InputChannel == nullptr);
-	InputChannel = new AJA::AJAInputChannel();
-	if (!InputChannel->Initialize(DeviceOptions, AjaOptions))
-	{
-		UE_LOG(LogAjaMedia, Warning, TEXT("The AJA port couldn't be opened."));
-		CurrentState = EMediaState::Error;
-		AjaThreadNewState = EMediaState::Error;
-		delete InputChannel;
-		InputChannel = nullptr;
-	}
-
-	// Setup our different supported channels based on source settings
-	SetupSampleChannels();
-
-	// configure format information for base class
-	AudioTrackFormat.BitsPerSample = 32;
-	AudioTrackFormat.NumChannels = 0;
-	AudioTrackFormat.SampleRate = 48000;
-	AudioTrackFormat.TypeName = FString(TEXT("PCM"));
-
-	// finalize
-	CurrentState = EMediaState::Preparing;
-	AjaThreadNewState = EMediaState::Preparing;
-	EventSink.ReceiveMediaEvent(EMediaEvent::MediaConnecting);
-
-#if WITH_EDITOR
-	if (FEngineAnalytics::IsAvailable())
-	{
-		TArray<FAnalyticsEventAttribute> EventAttributes;
-		
-		const int64 ResolutionWidth = Options->GetMediaOption( FMediaIOCoreMediaOption::ResolutionWidth, (int64)1920);
-		const int64 ResolutionHeight = Options->GetMediaOption( FMediaIOCoreMediaOption::ResolutionHeight, (int64)1080);
-		
-		EventAttributes.Add(FAnalyticsEventAttribute(TEXT("ResolutionWidth"), FString::Printf(TEXT("%d"), ResolutionWidth)));
-		EventAttributes.Add(FAnalyticsEventAttribute(TEXT("ResolutionHeight"), FString::Printf(TEXT("%d"), ResolutionHeight)));
-		EventAttributes.Add(FAnalyticsEventAttribute(TEXT("FrameRate"), *VideoFrameRate.ToPrettyText().ToString()));
-		
-		FEngineAnalytics::GetProvider().RecordEvent(TEXT("MediaFramework.AjaSourceOpened"), EventAttributes);
-	}
-#endif
-
-	return true;
 }
 
 void FAjaMediaPlayer::Close()
@@ -779,6 +644,187 @@
 	Samples->AddVideo(InSample);
 }
 
+bool FAjaMediaPlayer::Open_Internal(const FString& Url, const IMediaOptions* Options, AJA::AJAInputOutputChannelOptions AjaOptions)
+{
+	AJA::AJADeviceOptions DeviceOptions(Options->GetMediaOption(AjaMediaOption::DeviceIndex, (int64)0));
+
+	AjaOptions.CallbackInterface = this;
+	AjaOptions.bOutput = false;
+	{
+		const EMediaIOTransportType TransportType = (EMediaIOTransportType)(Options->GetMediaOption(AjaMediaOption::TransportType, (int64)EMediaIOTransportType::SingleLink));
+		const EMediaIOQuadLinkTransportType QuadTransportType = (EMediaIOQuadLinkTransportType)(Options->GetMediaOption(AjaMediaOption::QuadTransportType, (int64)EMediaIOQuadLinkTransportType::SquareDivision));
+		switch (TransportType)
+		{
+		case EMediaIOTransportType::SingleLink:
+			AjaOptions.TransportType = AJA::ETransportType::TT_SdiSingle;
+			break;
+		case EMediaIOTransportType::DualLink:
+			AjaOptions.TransportType = AJA::ETransportType::TT_SdiDual;
+			break;
+		case EMediaIOTransportType::QuadLink:
+			AjaOptions.TransportType = QuadTransportType == EMediaIOQuadLinkTransportType::SquareDivision ? AJA::ETransportType::TT_SdiQuadSQ : AJA::ETransportType::TT_SdiQuadTSI;
+			break;
+		case EMediaIOTransportType::HDMI:
+			AjaOptions.TransportType = AJA::ETransportType::TT_Hdmi;
+			break;
+		}
+	}
+	{
+		const EMediaIOTimecodeFormat Timecode = (EMediaIOTimecodeFormat)(Options->GetMediaOption(AjaMediaOption::TimecodeFormat, (int64)EMediaIOTimecodeFormat::None));
+		bUseFrameTimecode = Timecode != EMediaIOTimecodeFormat::None;
+		AjaOptions.TimecodeFormat = AJA::ETimecodeFormat::TCF_None;
+		switch (Timecode)
+		{
+		case EMediaIOTimecodeFormat::None:
+			AjaOptions.TimecodeFormat = AJA::ETimecodeFormat::TCF_None;
+			break;
+		case EMediaIOTimecodeFormat::LTC:
+			AjaOptions.TimecodeFormat = AJA::ETimecodeFormat::TCF_LTC;
+			break;
+		case EMediaIOTimecodeFormat::VITC:
+			AjaOptions.TimecodeFormat = AJA::ETimecodeFormat::TCF_VITC1;
+			break;
+		default:
+			break;
+		}
+		bEncodeTimecodeInTexel = Options->GetMediaOption(AjaMediaOption::EncodeTimecodeInTexel, false);
+	}
+	{
+		const EAjaMediaAudioChannel AudioChannelOption = (EAjaMediaAudioChannel)(Options->GetMediaOption(AjaMediaOption::AudioChannel, (int64)EAjaMediaAudioChannel::Channel8));
+		AjaOptions.NumberOfAudioChannel = (AudioChannelOption == EAjaMediaAudioChannel::Channel8) ? 8 : 6;
+	}
+	{
+		if (!bAutoDetect)
+		{
+			AjaOptions.VideoFormatIndex = Options->GetMediaOption(AjaMediaOption::AjaVideoFormat, (int64)0);
+		}
+		LastVideoFormatIndex = AjaOptions.VideoFormatIndex;
+	}
+	{
+		AjaColorFormat = (EAjaMediaSourceColorFormat)(Options->GetMediaOption(AjaMediaOption::ColorFormat, (int64)EAjaMediaSourceColorFormat::YUV2_8bit));
+		switch (AjaColorFormat)
+		{
+		case EAjaMediaSourceColorFormat::YUV2_8bit:
+			if (AjaOptions.bUseKey)
+			{
+				AjaOptions.PixelFormat = AJA::EPixelFormat::PF_8BIT_ARGB;
+			}
+			else
+			{
+				AjaOptions.PixelFormat = AJA::EPixelFormat::PF_8BIT_YCBCR;
+			}
+			break;
+		case EAjaMediaSourceColorFormat::YUV_10bit:
+			if (AjaOptions.bUseKey)
+			{
+				AjaOptions.PixelFormat = AJA::EPixelFormat::PF_10BIT_RGB;
+			}
+			else
+			{
+				AjaOptions.PixelFormat = AJA::EPixelFormat::PF_10BIT_YCBCR;
+			}
+			break;
+		default:
+			AjaOptions.PixelFormat = AJA::EPixelFormat::PF_8BIT_ARGB;
+			break;
+		}
+
+		bIsSRGBInput = Options->GetMediaOption(AjaMediaOption::SRGBInput, false);
+	}
+	{
+		AjaOptions.bUseAncillary = bUseAncillary = Options->GetMediaOption(AjaMediaOption::CaptureAncillary, false);
+		AjaOptions.bUseAudio = bUseAudio = Options->GetMediaOption(AjaMediaOption::CaptureAudio, false);
+		AjaOptions.bUseVideo = bUseVideo = Options->GetMediaOption(AjaMediaOption::CaptureVideo, true);
+		AjaOptions.bUseAutoCirculating = Options->GetMediaOption(AjaMediaOption::CaptureWithAutoCirculating, true);
+		AjaOptions.bUseKey = false;
+		AjaOptions.bBurnTimecode = false;
+		AjaOptions.BurnTimecodePercentY = 80;
+
+		//Adjust supported sample types based on what's being captured
+		SupportedSampleTypes = AjaOptions.bUseVideo ? EMediaIOSampleType::Video : EMediaIOSampleType::None;
+		SupportedSampleTypes |= AjaOptions.bUseAudio ? EMediaIOSampleType::Audio : EMediaIOSampleType::None;
+		SupportedSampleTypes |= AjaOptions.bUseAncillary ? EMediaIOSampleType::Metadata : EMediaIOSampleType::None;
+		Samples->EnableTimedDataChannels(this, SupportedSampleTypes);
+	}
+
+	bVerifyFrameDropCount = Options->GetMediaOption(AjaMediaOption::LogDropFrame, true);
+	MaxNumAudioFrameBuffer = Options->GetMediaOption(AjaMediaOption::MaxAudioFrameBuffer, (int64)8);
+	MaxNumMetadataFrameBuffer = Options->GetMediaOption(AjaMediaOption::MaxAncillaryFrameBuffer, (int64)8);
+	MaxNumVideoFrameBuffer = Options->GetMediaOption(AjaMediaOption::MaxVideoFrameBuffer, (int64)8);
+
+	check(InputChannel == nullptr);
+	InputChannel = new AJA::AJAInputChannel();
+	if (!InputChannel->Initialize(DeviceOptions, AjaOptions))
+	{
+		UE_LOG(LogAjaMedia, Warning, TEXT("The AJA port couldn't be opened."));
+		CurrentState = EMediaState::Error;
+		AjaThreadNewState = EMediaState::Error;
+		delete InputChannel;
+		InputChannel = nullptr;
+	}
+
+	// Setup our different supported channels based on source settings
+	SetupSampleChannels();
+
+	// configure format information for base class
+	AudioTrackFormat.BitsPerSample = 32;
+	AudioTrackFormat.NumChannels = 0;
+	AudioTrackFormat.SampleRate = 48000;
+	AudioTrackFormat.TypeName = FString(TEXT("PCM"));
+
+	// finalize
+	CurrentState = EMediaState::Preparing;
+	AjaThreadNewState = EMediaState::Preparing;
+	EventSink.ReceiveMediaEvent(EMediaEvent::MediaConnecting);
+
+#if WITH_EDITOR
+	if (FEngineAnalytics::IsAvailable())
+	{
+		TArray<FAnalyticsEventAttribute> EventAttributes;
+
+		EventAttributes.Add(FAnalyticsEventAttribute(TEXT("ResolutionWidth"), FString::Printf(TEXT("%d"), VideoTrackFormat.Dim.X)));
+		EventAttributes.Add(FAnalyticsEventAttribute(TEXT("ResolutionHeight"), FString::Printf(TEXT("%d"), VideoTrackFormat.Dim.Y)));
+		EventAttributes.Add(FAnalyticsEventAttribute(TEXT("FrameRate"), *VideoFrameRate.ToPrettyText().ToString()));
+
+		FEngineAnalytics::GetProvider().RecordEvent(TEXT("MediaFramework.AjaSourceOpened"), EventAttributes);
+	}
+#endif
+
+	return true;
+}
+
+void FAjaMediaPlayer::OnAutoDetected(TArray<FAjaDeviceProvider::FMediaIOConfigurationWithTimecodeFormat> Configurations, FString Url, const IMediaOptions* Options)
+{
+	AJA::AJAInputOutputChannelOptions AjaOptions(TEXT("MediaPlayer"), Options->GetMediaOption(AjaMediaOption::PortIndex, (int64)0));
+	
+	if (Configurations.Num())
+	{
+		FAjaDeviceProvider::FMediaIOConfigurationWithTimecodeFormat Format = Configurations[0];
+			
+		{
+			VideoFrameRate = Format.Configuration.MediaMode.FrameRate;
+			FIntPoint Resolution = Format.Configuration.MediaMode.Resolution;
+
+			VideoTrackFormat.Dim = Resolution;
+			VideoTrackFormat.FrameRates = TRange<float>(VideoFrameRate.AsDecimal());
+			VideoTrackFormat.FrameRate = VideoFrameRate.AsDecimal();
+			VideoTrackFormat.TypeName = Format.Configuration.MediaMode.GetModeName().ToString();
+
+			AjaOptions.VideoFormatIndex = Format.Configuration.MediaMode.DeviceModeIdentifier;
+		}
+
+		//Setup base sampling settings
+		BaseSettings.FrameRate = VideoFrameRate;
+	}
+	else
+	{
+		Close();
+		return;
+	}
+
+	Open_Internal(Url, Options, AjaOptions);
+}
+
 bool FAjaMediaPlayer::SetRate(float Rate)
 {
 	if (FMath::IsNearlyEqual(Rate, 1.0f))
diff --git a/Engine/Plugins/Media/AjaMedia/Source/AjaMedia/Private/Player/AjaMediaPlayer.h b/Engine/Plugins/Media/AjaMedia/Source/AjaMedia/Private/Player/AjaMediaPlayer.h
--- a/Engine/Plugins/Media/AjaMedia/Source/AjaMedia/Private/Player/AjaMediaPlayer.h
+++ b/Engine/Plugins/Media/AjaMedia/Source/AjaMedia/Private/Player/AjaMediaPlayer.h
@@ -4,6 +4,7 @@
 
 #include "MediaIOCorePlayerBase.h"
 
+#include "AjaDeviceProvider.h"
 #include "AjaMediaPrivate.h"
 #include "AjaMediaSource.h"
 #include "HAL/CriticalSection.h"
@@ -116,6 +117,10 @@ protected:
 
 	virtual void AddVideoSample(const TSharedRef<FMediaIOCoreTextureSampleBase>& InSample) override;
 
+private:
+	bool Open_Internal(const FString& Url, const IMediaOptions* Options, AJA::AJAInputOutputChannelOptions AjaOptions);
+	void OnAutoDetected(TArray<FAjaDeviceProvider::FMediaIOConfigurationWithTimecodeFormat> Configurations, FString Url, const IMediaOptions* Options);
+
 private:
 
 	/** Audio, MetaData, Texture  sample object pool. */
@@ -186,4 +191,7 @@ private:
 	std::atomic<bool> bPauseRequested;
 
 	EAjaMediaSourceColorFormat AjaColorFormat = EAjaMediaSourceColorFormat::YUV2_8bit;
+
+	/** Device provider used to autodetect input format. */
+	TPimplPtr<class FAjaDeviceProvider> DeviceProvider;
 };
diff --git a/Engine/Plugins/Media/AjaMedia/Source/AjaMedia/Public/AjaMediaSource.h b/Engine/Plugins/Media/AjaMedia/Source/AjaMedia/Public/AjaMediaSource.h
--- a/Engine/Plugins/Media/AjaMedia/Source/AjaMedia/Public/AjaMediaSource.h
+++ b/Engine/Plugins/Media/AjaMedia/Source/AjaMedia/Public/AjaMediaSource.h
@@ -134,5 +134,7 @@ public:
 	virtual bool CanEditChange(const FProperty* InProperty) const override;
 	virtual void PostEditChangeChainProperty(struct FPropertyChangedChainEvent& PropertyChangedEvent) override;
 #endif //WITH_EDITOR
+
+	virtual bool SupportsFormatAutoDetection() const override { return true; }
 	//~ End UObject interface
 };
diff --git a/Engine/Plugins/Media/MediaIOFramework/Source/MediaIOCore/Private/Player/MediaIOCorePlayerBase.cpp b/Engine/Plugins/Media/MediaIOFramework/Source/MediaIOCore/Private/Player/MediaIOCorePlayerBase.cpp
--- a/Engine/Plugins/Media/MediaIOFramework/Source/MediaIOCore/Private/Player/MediaIOCorePlayerBase.cpp
+++ b/Engine/Plugins/Media/MediaIOFramework/Source/MediaIOCore/Private/Player/MediaIOCorePlayerBase.cpp
@@ -84,6 +84,7 @@ FMediaIOCorePlayerBase::FMediaIOCorePlayerBase(IMediaEventSink& InEventSink)
 	, VideoFrameRate(30, 1)
 	, bUseTimeSynchronization(false)
 	, bWarnedIncompatibleFrameRate(false)
+	, bAutoDetect(false)
 	, FrameDelay(0)
 	, TimeDelay(0.0)
 	, PreviousFrameTimespan(FTimespan::Zero())
@@ -473,6 +474,7 @@ bool FMediaIOCorePlayerBase::ReadMediaOptions(const IMediaOptions* Options)
 	bUseTimeSynchronization = Options->GetMediaOption(TimeSynchronizableMedia::UseTimeSynchronizatioOption, false);
 	FrameDelay = Options->GetMediaOption(TimeSynchronizableMedia::FrameDelay, (int64)0);
 	TimeDelay = Options->GetMediaOption(TimeSynchronizableMedia::TimeDelay, 0.0);
+	bAutoDetect = Options->GetMediaOption(TimeSynchronizableMedia::AutoDetect, false);
 
 	{
 		int32 Numerator = Options->GetMediaOption(FMediaIOCoreMediaOption::FrameRateNumerator, (int64)30);
diff --git a/Engine/Plugins/Media/MediaIOFramework/Source/MediaIOCore/Public/MediaIOCorePlayerBase.h b/Engine/Plugins/Media/MediaIOFramework/Source/MediaIOCore/Public/MediaIOCorePlayerBase.h
--- a/Engine/Plugins/Media/MediaIOFramework/Source/MediaIOCore/Public/MediaIOCorePlayerBase.h
+++ b/Engine/Plugins/Media/MediaIOFramework/Source/MediaIOCore/Public/MediaIOCorePlayerBase.h
@@ -229,6 +229,9 @@ protected:
 	/** Warn when the video frame rate is not the same as the engine's frame rate. */
 	bool bWarnedIncompatibleFrameRate;
 
+	/** Whether we are using autodetection. */
+	bool bAutoDetect;
+
 	/** When using Time Synchronization (TC synchronization), how many frame back of a delay would you like. */
 	int32 FrameDelay;
 
diff --git a/Engine/Plugins/Media/MediaIOFramework/Source/MediaIOEditor/MediaIOEditor.Build.cs b/Engine/Plugins/Media/MediaIOFramework/Source/MediaIOEditor/MediaIOEditor.Build.cs
--- a/Engine/Plugins/Media/MediaIOFramework/Source/MediaIOEditor/MediaIOEditor.Build.cs
+++ b/Engine/Plugins/Media/MediaIOFramework/Source/MediaIOEditor/MediaIOEditor.Build.cs
@@ -18,6 +18,12 @@ namespace UnrealBuildTool.Rules
 					"TimeManagement",
 					"UnrealEd",
 				});
+
+			PrivateDependencyModuleNames.AddRange(
+				new string[]
+				{
+					"MediaAssets"
+				});
 		}
 	}
 }
diff --git a/Engine/Plugins/Media/MediaIOFramework/Source/MediaIOEditor/Private/Customizations/MediaIOConfigurationCustomization.cpp b/Engine/Plugins/Media/MediaIOFramework/Source/MediaIOEditor/Private/Customizations/MediaIOConfigurationCustomization.cpp
--- a/Engine/Plugins/Media/MediaIOFramework/Source/MediaIOEditor/Private/Customizations/MediaIOConfigurationCustomization.cpp
+++ b/Engine/Plugins/Media/MediaIOFramework/Source/MediaIOEditor/Private/Customizations/MediaIOConfigurationCustomization.cpp
@@ -11,6 +11,8 @@
 #include "IMediaIOCoreModule.h"
 #include "MediaIOPermutationsSelectorBuilder.h"
 #include "ObjectEditorUtils.h"
+#include "TimeSynchronizableMediaSource.h"
+#include "Widgets/Input/SCheckBox.h"
 #include "Widgets/Input/SComboButton.h"
 #include "Widgets/Text/STextBlock.h"
 #include "Widgets/SBoxPanel.h"
@@ -22,7 +24,7 @@
 
 TSharedRef<IPropertyTypeCustomization> FMediaIOConfigurationCustomization::MakeInstance()
 {
-	return MakeShareable(new FMediaIOConfigurationCustomization);
+	return MakeShared<FMediaIOConfigurationCustomization>();
 }
 
 TAttribute<FText> FMediaIOConfigurationCustomization::GetContentText()
@@ -72,6 +74,56 @@ TSharedRef<SWidget> FMediaIOConfigurationCustomization::HandleSourceComboButtonM
 		return false;
 	};
 
+	TSharedRef<SWidget> EnforceWidget = SNullWidget::NullWidget;
+
+	if (bIsInput)
+	{
+		bool bSupportsAutoDetect = true;
+		for (UObject* Object : GetCustomizedObjects())
+		{
+			if (Object && Object->IsA<UTimeSynchronizableMediaSource>())
+			{
+				const UTimeSynchronizableMediaSource* MediaSource = CastChecked<UTimeSynchronizableMediaSource>(Object);
+				bSupportsAutoDetect &= MediaSource->SupportsFormatAutoDetection();
+				CustomizedSources.Add(CastChecked<UTimeSynchronizableMediaSource>(Object));
+			}
+			else
+			{
+				bSupportsAutoDetect = false;
+			}
+		}
+
+		if (bSupportsAutoDetect)
+		{
+			EnforceWidget = SNew(SHorizontalBox)
+				+ SHorizontalBox::Slot()
+				.Padding(4.f)
+				.AutoWidth()
+				[
+					SNew(STextBlock)
+					.Text(NSLOCTEXT("MediaPlayerEditor", "EnforceLabel", "Enforce format"))
+				]
+			+ SHorizontalBox::Slot()
+				.Padding(4.f)
+				.AutoWidth()
+				[
+					SNew(SCheckBox)
+					.ToolTipText(NSLOCTEXT("MediaPlayerEditor", "EnforceTooltip", "Enforce a specific configuration. (Note that this disables format autodetection)"))
+					.IsChecked(this, &FMediaIOConfigurationCustomization::GetEnforceCheckboxState)
+					.OnCheckStateChanged(this, &FMediaIOConfigurationCustomization::SetEnforceCheckboxState)
+				];
+		}		
+
+	}
+
+	auto GetExtensions = [EnforceWidget](TArray<TSharedRef<SWidget>>& OutWidgets)
+	{
+		if (EnforceWidget != SNullWidget::NullWidget)
+		{
+			OutWidgets.Add(EnforceWidget);
+		}
+	};
+
 	using TSelection = SMediaPermutationsSelector<FMediaIOConfiguration, FMediaIOPermutationsSelectorBuilder>;
 	TSelection::FArguments Arguments;
 	Arguments
@@ -79,6 +131,7 @@ TSharedRef<SWidget> FMediaIOConfigurationCustomization::HandleSourceComboButtonM
 		.SelectedPermutation(SelectedConfiguration)
 		.OnSelectionChanged(this, &FMediaIOConfigurationCustomization::OnSelectionChanged)
 		.OnButtonClicked(this, &FMediaIOConfigurationCustomization::OnButtonClicked)
+		.OnGetExtensions_Lambda(GetExtensions)
 		+ TSelection::Column(FMediaIOPermutationsSelectorBuilder::NAME_DeviceIdentifier)
 		.Label(LOCTEXT("DeviceLabel", "Device"));
 
@@ -95,10 +148,15 @@ TSharedRef<SWidget> FMediaIOConfigurationCustomization::HandleSourceComboButtonM
 	Arguments
 		+ TSelection::Column(FMediaIOPermutationsSelectorBuilder::NAME_Resolution)
 		.Label(LOCTEXT("ResolutionLabel", "Resolution"))
+		.IsColumnVisible_Raw(this, &FMediaIOConfigurationCustomization::ShowAdvancedColumns)
 		+ TSelection::Column(FMediaIOPermutationsSelectorBuilder::NAME_Standard)
 		.Label(LOCTEXT("StandardLabel", "Standard"))
+		.IsColumnVisible_Raw(this, &FMediaIOConfigurationCustomization::ShowAdvancedColumns)
 		+ TSelection::Column(FMediaIOPermutationsSelectorBuilder::NAME_FrameRate)
-		.Label(LOCTEXT("FrameRateLabel", "Frame Rate"));
+		.Label(LOCTEXT("FrameRateLabel", "Frame Rate"))
+		.IsColumnVisible_Raw(this, &FMediaIOConfigurationCustomization::ShowAdvancedColumns);
+
+
 
 	TSharedRef<TSelection> Selector = SNew(TSelection) = Arguments;
 	PermutationSelector = Selector;
@@ -126,4 +184,56 @@ FReply FMediaIOConfigurationCustomization::OnButtonClicked() const
 	return FReply::Handled();
 }
 
+bool FMediaIOConfigurationCustomization::ShowAdvancedColumns(FName ColumnName, const TArray<FMediaIOConfiguration>& UniquePermutationsForThisColumn) const
+{
+	bool bAutoDetectInput = true;
+
+	for (const TWeakObjectPtr<UTimeSynchronizableMediaSource>& WeakMediaSource : CustomizedSources)
+	{
+		if (UTimeSynchronizableMediaSource* MediaSource = WeakMediaSource.Get())
+		{
+			if (!MediaSource->SupportsFormatAutoDetection())
+			{
+				return true;
+			}
+			if (!MediaSource->bAutoDetectInput)
+			{
+				bAutoDetectInput = false;
+				break;
+			}
+		}
+	}
+
+	return !bAutoDetectInput;
+}
+
+ECheckBoxState FMediaIOConfigurationCustomization::GetEnforceCheckboxState() const
+{
+	bool bAllAutoDetect = true;
+
+	for (const TWeakObjectPtr<UTimeSynchronizableMediaSource>& WeakMediaSource : CustomizedSources)
+	{
+		if(UTimeSynchronizableMediaSource * MediaSource = WeakMediaSource.Get(); !MediaSource->bAutoDetectInput)
+		{
+			bAllAutoDetect = false;
+			break;
+		}
+	}
+
+	return bAllAutoDetect ? ECheckBoxState::Unchecked : ECheckBoxState::Checked;
+}
+
+void FMediaIOConfigurationCustomization::SetEnforceCheckboxState(ECheckBoxState CheckboxState)
+{
+	bool bEnforceFormat = CheckboxState == ECheckBoxState::Checked;
+
+	for (const TWeakObjectPtr<UTimeSynchronizableMediaSource>& WeakMediaSource : CustomizedSources)
+	{
+		if (UTimeSynchronizableMediaSource* MediaSource = WeakMediaSource.Get())
+		{
+			MediaSource->bAutoDetectInput = !bEnforceFormat;
+		}
+	}
+}
+
 #undef LOCTEXT_NAMESPACE
diff --git a/Engine/Plugins/Media/MediaIOFramework/Source/MediaIOEditor/Private/Customizations/MediaIOCustomizationBase.cpp b/Engine/Plugins/Media/MediaIOFramework/Source/MediaIOEditor/Private/Customizations/MediaIOCustomizationBase.cpp
--- a/Engine/Plugins/Media/MediaIOFramework/Source/MediaIOEditor/Private/Customizations/MediaIOCustomizationBase.cpp
+++ b/Engine/Plugins/Media/MediaIOFramework/Source/MediaIOEditor/Private/Customizations/MediaIOCustomizationBase.cpp
@@ -117,6 +117,16 @@ void FMediaIOCustomizationBase::CustomizeChildren(TSharedRef<IPropertyHandle> In
 	}
 }
 
+TArray<UObject*> FMediaIOCustomizationBase::GetCustomizedObjects() const
+{
+	TArray<UObject*> Objects;
+	if (MediaProperty)
+	{
+		MediaProperty->GetOuterObjects(Objects);
+	}
+	return Objects;
+}
+
 void FMediaIOCustomizationBase::AssignValueImpl(const void* NewValue) const
 {
 	if (FStructProperty* StructProperty = CastField<FStructProperty>(MediaProperty->GetProperty()))
diff --git a/Engine/Plugins/Media/MediaIOFramework/Source/MediaIOEditor/Public/Customizations/MediaIOConfigurationCustomization.h b/Engine/Plugins/Media/MediaIOFramework/Source/MediaIOEditor/Public/Customizations/MediaIOConfigurationCustomization.h
--- a/Engine/Plugins/Media/MediaIOFramework/Source/MediaIOEditor/Public/Customizations/MediaIOConfigurationCustomization.h
+++ b/Engine/Plugins/Media/MediaIOFramework/Source/MediaIOEditor/Public/Customizations/MediaIOConfigurationCustomization.h
@@ -18,9 +18,14 @@ private:
 	virtual TAttribute<FText> GetContentText() override;
 	virtual TSharedRef<SWidget> HandleSourceComboButtonMenuContent() override;
 
+	ECheckBoxState GetEnforceCheckboxState() const;
+	void SetEnforceCheckboxState(ECheckBoxState CheckboxState);
+	bool ShowAdvancedColumns(FName ColumnName, const TArray<FMediaIOConfiguration>& UniquePermutationsForThisColumn) const;
+
 	void OnSelectionChanged(FMediaIOConfiguration SelectedItem);
 	FReply OnButtonClicked() const;
 
 	TWeakPtr<SWidget> PermutationSelector;
 	FMediaIOConfiguration SelectedConfiguration;
+	TArray<TWeakObjectPtr<class UTimeSynchronizableMediaSource>> CustomizedSources;
 };
diff --git a/Engine/Plugins/Media/MediaIOFramework/Source/MediaIOEditor/Public/Customizations/MediaIOCustomizationBase.h b/Engine/Plugins/Media/MediaIOFramework/Source/MediaIOEditor/Public/Customizations/MediaIOCustomizationBase.h
--- a/Engine/Plugins/Media/MediaIOFramework/Source/MediaIOEditor/Public/Customizations/MediaIOCustomizationBase.h
+++ b/Engine/Plugins/Media/MediaIOFramework/Source/MediaIOEditor/Public/Customizations/MediaIOCustomizationBase.h
@@ -36,6 +36,7 @@ protected:
 			check(MediaValue);
 			return MediaValue;
 	}
+	
 	template<class T>
 	void AssignValue(const T& NewValue) const
 	{
@@ -45,6 +46,7 @@ protected:
 	virtual TAttribute<FText> GetContentText() = 0;
 	virtual TSharedRef<SWidget> HandleSourceComboButtonMenuContent() = 0;
 	TSharedPtr<IPropertyHandle> GetMediaProperty() const { return MediaProperty; };
+	TArray<UObject*> GetCustomizedObjects() const;
 
 	FName DeviceProviderName;
 
diff --git a/Engine/Plugins/Media/MediaIOFramework/Source/MediaIOEditor/Public/Widgets/SMediaPermutationsSelector.h b/Engine/Plugins/Media/MediaIOFramework/Source/MediaIOEditor/Public/Widgets/SMediaPermutationsSelector.h
--- a/Engine/Plugins/Media/MediaIOFramework/Source/MediaIOEditor/Public/Widgets/SMediaPermutationsSelector.h
+++ b/Engine/Plugins/Media/MediaIOFramework/Source/MediaIOEditor/Public/Widgets/SMediaPermutationsSelector.h
@@ -75,6 +75,11 @@ public:
 	 */
 	DECLARE_DELEGATE_RetVal_TwoParams(bool, FIsColumnVisible, FName /*ColumnName*/, const TArray<ItemType>& /*UniquePermutationsForThisColumn */);
 
+	/**
+	 * A delegate called to add additional widgets to the left of the apply button.
+	 */
+	DECLARE_DELEGATE_OneParam(FOnGetExtensions, TArray<TSharedRef<SWidget>>& /*OutExtensions*/);
+
 public:
 	/** Describes a single column */
 	class FColumn
@@ -135,6 +140,8 @@ public:
 		SLATE_EVENT(FOnClicked, OnButtonClicked)
 		/** Delegate to invoke when selection changes. */
 		SLATE_EVENT(FOnSelectionChanged, OnSelectionChanged)
+		/** Delegate to invoke when gathering extensions. */
+		SLATE_EVENT(FOnGetExtensions, OnGetExtensions)
 	SLATE_END_ARGS()
 
 	void Construct(const FArguments& InArgs);
diff --git a/Engine/Plugins/Media/MediaIOFramework/Source/MediaIOEditor/Public/Widgets/SMediaPermutationsSelector.inl b/Engine/Plugins/Media/MediaIOFramework/Source/MediaIOEditor/Public/Widgets/SMediaPermutationsSelector.inl
--- a/Engine/Plugins/Media/MediaIOFramework/Source/MediaIOEditor/Public/Widgets/SMediaPermutationsSelector.inl
+++ b/Engine/Plugins/Media/MediaIOFramework/Source/MediaIOEditor/Public/Widgets/SMediaPermutationsSelector.inl
@@ -44,6 +44,22 @@ void SMediaPermutationsSelector<ItemType, ItemBuilder>::Construct(const FArgumen
 
 	SAssignNew(ColumnContainer, SHorizontalBox);
 
+	TArray<TSharedRef<SWidget>> Extensions;
+	InArgs._OnGetExtensions.ExecuteIfBound(Extensions);
+	const bool bIsVisible = Extensions.Num() != 0;
+
+	TSharedRef<SHorizontalBox> ExtensionWidgetContainer = SNew(SHorizontalBox)
+		.Visibility(bIsVisible ? EVisibility::Visible : EVisibility::Collapsed);
+
+	for (int32 Index = Extensions.Num() - 1; Index >= 0; Index--)
+	{
+		ExtensionWidgetContainer->AddSlot()
+			.AutoWidth()
+			[
+				Extensions[Index]
+			];
+	}
+
 	BuildColumns(0);
 
 	ChildSlot
@@ -69,6 +85,12 @@ void SMediaPermutationsSelector<ItemType, ItemBuilder>::Construct(const FArgumen
 				+ SHorizontalBox::Slot()
 				.Padding(4.f)
 				.AutoWidth()
+				[
+					ExtensionWidgetContainer
+				]
+				+ SHorizontalBox::Slot()
+				.Padding(4.f)
+				.AutoWidth()
 				[
 					InArgs._OverrideButtonWidget.IsValid() ?
 						InArgs._OverrideButtonWidget.ToSharedRef()
@@ -151,12 +173,6 @@ void SMediaPermutationsSelector<ItemType, ItemBuilder>::BuildColumns(int32 Start
 			}
 		}
 
-		// Only show the column if the user desire it
-		if (!IsColumnVisible(Column, UniqueItemsForColumnIndexes))
-		{
-			continue;
-		}
-
 		// Sort the column items
 		UniqueItemsForColumnIndexes.Sort([this, ColumnName](int32 Left, int32 Right) { return Less(ColumnName, Left, Right); });
 
@@ -180,6 +196,7 @@ void SMediaPermutationsSelector<ItemType, ItemBuilder>::BuildColumns(int32 Start
 
 		// Create the widget
 		SAssignNew(Column.Widget, SVerticalBox)
+		.Visibility_Lambda([this, Column, UniqueItemsForColumnIndexes]() { return IsColumnVisible(Column, UniqueItemsForColumnIndexes) ? EVisibility::Visible : EVisibility::Collapsed; })
 		+ SVerticalBox::Slot()
 		.AutoHeight()
 		.HAlign(HAlign_Left)
diff --git a/Engine/Source/Runtime/MediaAssets/Private/Assets/TimeSynchronizableMediaSource.cpp b/Engine/Source/Runtime/MediaAssets/Private/Assets/TimeSynchronizableMediaSource.cpp
--- a/Engine/Source/Runtime/MediaAssets/Private/Assets/TimeSynchronizableMediaSource.cpp
+++ b/Engine/Source/Runtime/MediaAssets/Private/Assets/TimeSynchronizableMediaSource.cpp
@@ -21,6 +21,11 @@ bool UTimeSynchronizableMediaSource::GetMediaOption(const FName& Key, bool Defau
 	{
 		return bUseTimeSynchronization;
 	}
+	
+	if (Key == TimeSynchronizableMedia::AutoDetect)
+	{
+		return bAutoDetectInput;
+	}
 
 	return Super::GetMediaOption(Key, DefaultValue);
 }
@@ -54,7 +59,8 @@ bool UTimeSynchronizableMediaSource::HasMediaOption(const FName& Key) const
 {
 	if (Key == TimeSynchronizableMedia::UseTimeSynchronizatioOption
 		|| Key == TimeSynchronizableMedia::FrameDelay
-		|| Key == TimeSynchronizableMedia::TimeDelay)
+		|| Key == TimeSynchronizableMedia::TimeDelay
+		|| Key == TimeSynchronizableMedia::AutoDetect)
 	{
 		return true;
 	}
diff --git a/Engine/Source/Runtime/MediaAssets/Public/TimeSynchronizableMediaSource.h b/Engine/Source/Runtime/MediaAssets/Public/TimeSynchronizableMediaSource.h
--- a/Engine/Source/Runtime/MediaAssets/Public/TimeSynchronizableMediaSource.h
+++ b/Engine/Source/Runtime/MediaAssets/Public/TimeSynchronizableMediaSource.h
@@ -13,6 +13,7 @@ namespace TimeSynchronizableMedia
 	static FName UseTimeSynchronizatioOption("UseTimeSynchronization");
 	static FName FrameDelay("FrameDelay");
 	static FName TimeDelay("TimeDelay");
+	static FName AutoDetect("AutoDetect");
 }
 
 /**
@@ -45,6 +46,10 @@ public:
 	UPROPERTY(EditAnywhere, Category=Synchronization, meta=(EditCondition="!bUseTimeSynchronization", ForceUnits=s))
 	double TimeDelay;
 
+	/** Whether to autodetect the input or not. */
+	UPROPERTY()
+	bool bAutoDetectInput = true;
+
 public:
 	//~ IMediaOptions interface
 	using Super::GetMediaOption;
@@ -53,4 +58,6 @@ public:
 	virtual double GetMediaOption(const FName& Key, double DefaultValue) const override;
 	virtual FString GetMediaOption(const FName& Key, const FString& DefaultValue) const override;
 	virtual bool HasMediaOption(const FName& Key) const override;
+
+	virtual bool SupportsFormatAutoDetection() const { return false; }
 };
