Add autodetect for aja timecode

diff --git a/Engine/Plugins/Media/AjaMedia/Source/AjaMedia/Private/Aja/AjaDeviceProvider.cpp b/Engine/Plugins/Media/AjaMedia/Source/AjaMedia/Private/Aja/AjaDeviceProvider.cpp
--- a/Engine/Plugins/Media/AjaMedia/Source/AjaMedia/Private/Aja/AjaDeviceProvider.cpp
+++ b/Engine/Plugins/Media/AjaMedia/Source/AjaMedia/Private/Aja/AjaDeviceProvider.cpp
@@ -245,13 +245,11 @@ public:
 //--------------------------------------------------------------------
 FAjaDeviceProvider::FAjaDeviceProvider()
 {
-	AutoDetectCallback = nullptr;
 }
 
 
 FAjaDeviceProvider::~FAjaDeviceProvider()
 {
-	delete AutoDetectCallback;
 }
 
 FName FAjaDeviceProvider::GetProviderName()
@@ -304,11 +302,15 @@ void FAjaDeviceProvider::AutoDetectConfiguration(FOnConfigurationAutoDetected On
 
 	if (OnAutoDetected.IsBound())
 	{
-		delete AutoDetectCallback;
-		AutoDetectCallback = new FAJAAutoDetectChannelCallback(OnAutoDetected);
+		AutoDetectCallback = MakeUnique<FAJAAutoDetectChannelCallback>(OnAutoDetected);
 	}
 }
 
+void FAjaDeviceProvider::EndAutoDetectConfiguration()
+{
+	AutoDetectCallback.Reset();
+}
+
 TArray<FMediaIOConnection> FAjaDeviceProvider::GetConnections() const
 {
 	TArray<FMediaIOConnection> Results;
diff --git a/Engine/Plugins/Media/AjaMedia/Source/AjaMedia/Private/Assets/AjaMediaSource.cpp b/Engine/Plugins/Media/AjaMedia/Source/AjaMedia/Private/Assets/AjaMediaSource.cpp
--- a/Engine/Plugins/Media/AjaMedia/Source/AjaMedia/Private/Assets/AjaMediaSource.cpp
+++ b/Engine/Plugins/Media/AjaMedia/Source/AjaMedia/Private/Assets/AjaMediaSource.cpp
@@ -9,8 +9,7 @@
 #include "UObject/EnterpriseObjectVersion.h"
 
 UAjaMediaSource::UAjaMediaSource()
-	: TimecodeFormat(EMediaIOTimecodeFormat::None)
-	, bCaptureWithAutoCirculating(true)
+	: bCaptureWithAutoCirculating(true)
 	, bCaptureAncillary(false)
 	, MaxNumAncillaryFrameBuffer(8)
 	, bCaptureAudio(false)
@@ -101,7 +100,7 @@ int64 UAjaMediaSource::GetMediaOption(const FName& Key, int64 DefaultValue) cons
 	}
 	if (Key == AjaMediaOption::TimecodeFormat)
 	{
-		return (int64)TimecodeFormat;
+		return (int64)AutoDetectableTimecodeFormat;
 	}
 	if (Key == AjaMediaOption::MaxAncillaryFrameBuffer)
 	{
@@ -243,9 +242,9 @@ bool UAjaMediaSource::Validate() const
 		return false;
 	}
 
-	if (bUseTimeSynchronization && TimecodeFormat == EMediaIOTimecodeFormat::None)
+	if (bUseTimeSynchronization && AutoDetectableTimecodeFormat == EMediaIOAutoDetectableTimecodeFormat::None)
 	{
-		UE_LOG(LogAjaMedia, Warning, TEXT("The MediaSource '%s' use time synchronization but doesn't enabled the timecode."), *GetName());
+		UE_LOG(LogAjaMedia, Warning, TEXT("The MediaSource '%s' uses time synchronization but hasn't enabled the timecode."), *GetName());
 		return false;
 	}
 
@@ -276,12 +275,12 @@ bool UAjaMediaSource::CanEditChange(const FProperty* InProperty) const
 
 	if (InProperty->GetFName() == GET_MEMBER_NAME_CHECKED(UAjaMediaSource, bEncodeTimecodeInTexel))
 	{
-		return TimecodeFormat != EMediaIOTimecodeFormat::None && bCaptureVideo;
+		return AutoDetectableTimecodeFormat != EMediaIOAutoDetectableTimecodeFormat::None && bCaptureVideo;
 	}
 
 	if (InProperty->GetFName() == GET_MEMBER_NAME_CHECKED(UTimeSynchronizableMediaSource, bUseTimeSynchronization))
 	{
-		return TimecodeFormat != EMediaIOTimecodeFormat::None;
+		return AutoDetectableTimecodeFormat != EMediaIOAutoDetectableTimecodeFormat::None;
 	}
 
 	return true;
@@ -289,9 +288,9 @@ bool UAjaMediaSource::CanEditChange(const FProperty* InProperty) const
 
 void UAjaMediaSource::PostEditChangeChainProperty(struct FPropertyChangedChainEvent& PropertyChangedEvent)
 {
-	if (PropertyChangedEvent.Property->GetFName() == GET_MEMBER_NAME_CHECKED(UAjaMediaSource, TimecodeFormat))
+	if (PropertyChangedEvent.Property->GetFName() == GET_MEMBER_NAME_CHECKED(UAjaMediaSource, AutoDetectableTimecodeFormat))
 	{
-		if (TimecodeFormat == EMediaIOTimecodeFormat::None)
+		if (AutoDetectableTimecodeFormat == EMediaIOAutoDetectableTimecodeFormat::None)
 		{
 			bUseTimeSynchronization = false;
 			bEncodeTimecodeInTexel = false;
@@ -301,3 +300,28 @@ void UAjaMediaSource::PostEditChangeChainProperty(struct FPropertyChangedChainEv
 	Super::PostEditChangeChainProperty(PropertyChangedEvent);
 }
 #endif //WITH_EDITOR
+
+void UAjaMediaSource::PostLoad()
+{
+	Super::PostLoad();
+
+#if WITH_EDITORONLY_DATA
+	PRAGMA_DISABLE_DEPRECATION_WARNINGS
+	if (TimecodeFormat_DEPRECATED != EMediaIOTimecodeFormat::None)
+	{
+		switch (TimecodeFormat_DEPRECATED)
+		{
+			case EMediaIOTimecodeFormat::LTC:
+				AutoDetectableTimecodeFormat = EMediaIOAutoDetectableTimecodeFormat::LTC;
+				break;
+			case EMediaIOTimecodeFormat::VITC:
+				AutoDetectableTimecodeFormat = EMediaIOAutoDetectableTimecodeFormat::VITC;
+				break;
+			default:
+				break;
+		}
+		TimecodeFormat_DEPRECATED = EMediaIOTimecodeFormat::None;
+	}
+	PRAGMA_ENABLE_DEPRECATION_WARNINGS
+#endif
+}
diff --git a/Engine/Plugins/Media/AjaMedia/Source/AjaMedia/Private/Player/AjaMediaPlayer.cpp b/Engine/Plugins/Media/AjaMedia/Source/AjaMedia/Private/Player/AjaMediaPlayer.cpp
--- a/Engine/Plugins/Media/AjaMedia/Source/AjaMedia/Private/Player/AjaMediaPlayer.cpp
+++ b/Engine/Plugins/Media/AjaMedia/Source/AjaMedia/Private/Player/AjaMediaPlayer.cpp
@@ -105,11 +105,19 @@ bool FAjaMediaPlayer::Open(const FString& Url, const IMediaOptions* Options)
 	{
 		return false;
 	}
+
+	const EMediaIOAutoDetectableTimecodeFormat Timecode = (EMediaIOAutoDetectableTimecodeFormat)(Options->GetMediaOption(AjaMediaOption::TimecodeFormat, (int64)EMediaIOAutoDetectableTimecodeFormat::None));
+	const bool bAutoDetectTimecode = Timecode == EMediaIOAutoDetectableTimecodeFormat::Auto;
+	const bool bAutoDetectVideoFormat = bAutoDetect;
+
+	if (!bAutoDetectTimecode)
+	{
+		TimecodeFormat = UE::MediaIO::FromAutoDetectableTimecodeFormat(Timecode);
+	}
 	
-	if (bAutoDetect)
+	if (bAutoDetectTimecode || bAutoDetectVideoFormat)
 	{
-		AJA::AJADeviceOptions DeviceOptions(Options->GetMediaOption(AjaMediaOption::DeviceIndex, (int64)0));
-		DeviceProvider->AutoDetectConfiguration(FAjaDeviceProvider::FOnConfigurationAutoDetected::CreateRaw(this, &FAjaMediaPlayer::OnAutoDetected, Url, Options));
+		DeviceProvider->AutoDetectConfiguration(FAjaDeviceProvider::FOnConfigurationAutoDetected::CreateRaw(this, &FAjaMediaPlayer::OnAutoDetected, Url, Options, bAutoDetectVideoFormat, bAutoDetectTimecode));
 		return true;
 	}
 	else
@@ -670,10 +678,9 @@ bool FAjaMediaPlayer::Open_Internal(const FString& Url, const IMediaOptions* Opt
 		}
 	}
 	{
-		const EMediaIOTimecodeFormat Timecode = (EMediaIOTimecodeFormat)(Options->GetMediaOption(AjaMediaOption::TimecodeFormat, (int64)EMediaIOTimecodeFormat::None));
-		bUseFrameTimecode = Timecode != EMediaIOTimecodeFormat::None;
+		bUseFrameTimecode = TimecodeFormat != EMediaIOTimecodeFormat::None;
 		AjaOptions.TimecodeFormat = AJA::ETimecodeFormat::TCF_None;
-		switch (Timecode)
+		switch (TimecodeFormat)
 		{
 		case EMediaIOTimecodeFormat::None:
 			AjaOptions.TimecodeFormat = AJA::ETimecodeFormat::TCF_None;
@@ -793,7 +800,7 @@ bool FAjaMediaPlayer::Open_Internal(const FString& Url, const IMediaOptions* Opt
 	return true;
 }
 
-void FAjaMediaPlayer::OnAutoDetected(TArray<FAjaDeviceProvider::FMediaIOConfigurationWithTimecodeFormat> Configurations, FString Url, const IMediaOptions* Options)
+void FAjaMediaPlayer::OnAutoDetected(TArray<FAjaDeviceProvider::FMediaIOConfigurationWithTimecodeFormat> Configurations, FString Url, const IMediaOptions* Options, bool bAutoDetectVideoFormat, bool bAutoDetectTimecodeFormat)
 {
 	AJA::AJAInputOutputChannelOptions AjaOptions(TEXT("MediaPlayer"), Options->GetMediaOption(AjaMediaOption::PortIndex, (int64)0));
 	
@@ -801,6 +808,7 @@ void FAjaMediaPlayer::OnAutoDetected(TArray<FAjaDeviceProvider::FMediaIOConfigur
 	{
 		FAjaDeviceProvider::FMediaIOConfigurationWithTimecodeFormat Format = Configurations[0];
 			
+		if (bAutoDetectVideoFormat)
 		{
 			VideoFrameRate = Format.Configuration.MediaMode.FrameRate;
 			FIntPoint Resolution = Format.Configuration.MediaMode.Resolution;
@@ -809,8 +817,14 @@ void FAjaMediaPlayer::OnAutoDetected(TArray<FAjaDeviceProvider::FMediaIOConfigur
 			VideoTrackFormat.FrameRates = TRange<float>(VideoFrameRate.AsDecimal());
 			VideoTrackFormat.FrameRate = VideoFrameRate.AsDecimal();
 			VideoTrackFormat.TypeName = Format.Configuration.MediaMode.GetModeName().ToString();
-
+			
 			AjaOptions.VideoFormatIndex = Format.Configuration.MediaMode.DeviceModeIdentifier;
+			
+		}
+
+		if (bAutoDetectTimecodeFormat)
+		{
+			TimecodeFormat = Format.TimecodeFormat;
 		}
 
 		//Setup base sampling settings
diff --git a/Engine/Plugins/Media/AjaMedia/Source/AjaMedia/Private/Player/AjaMediaPlayer.h b/Engine/Plugins/Media/AjaMedia/Source/AjaMedia/Private/Player/AjaMediaPlayer.h
--- a/Engine/Plugins/Media/AjaMedia/Source/AjaMedia/Private/Player/AjaMediaPlayer.h
+++ b/Engine/Plugins/Media/AjaMedia/Source/AjaMedia/Private/Player/AjaMediaPlayer.h
@@ -119,7 +119,7 @@ protected:
 
 private:
 	bool Open_Internal(const FString& Url, const IMediaOptions* Options, AJA::AJAInputOutputChannelOptions AjaOptions);
-	void OnAutoDetected(TArray<FAjaDeviceProvider::FMediaIOConfigurationWithTimecodeFormat> Configurations, FString Url, const IMediaOptions* Options);
+	void OnAutoDetected(TArray<FAjaDeviceProvider::FMediaIOConfigurationWithTimecodeFormat> Configurations, FString Url, const IMediaOptions* Options, bool bAutoDetectVideoFormat, bool bAutoDetectTimecodeFormat);
 
 private:
 
@@ -194,4 +194,7 @@ private:
 
 	/** Device provider used to autodetect input format. */
 	TPimplPtr<class FAjaDeviceProvider> DeviceProvider;
+
+	/** Autodetected or specified Timecode format. */
+	EMediaIOTimecodeFormat TimecodeFormat = EMediaIOTimecodeFormat::None; 
 };
diff --git a/Engine/Plugins/Media/AjaMedia/Source/AjaMedia/Public/AjaDeviceProvider.h b/Engine/Plugins/Media/AjaMedia/Source/AjaMedia/Public/AjaDeviceProvider.h
--- a/Engine/Plugins/Media/AjaMedia/Source/AjaMedia/Public/AjaDeviceProvider.h
+++ b/Engine/Plugins/Media/AjaMedia/Source/AjaMedia/Public/AjaDeviceProvider.h
@@ -109,6 +109,8 @@ public:
 
 	/** Auto detect sources that are currently streaming to the device */
 	void AutoDetectConfiguration(FOnConfigurationAutoDetected OnAutoDetected);
+	void EndAutoDetectConfiguration();
+
 
 public:
 	virtual FName GetFName() override;
@@ -129,5 +131,5 @@ public:
 	virtual FMediaIOOutputConfiguration GetDefaultOutputConfiguration() const override;
 
 private:
-	FAJAAutoDetectChannelCallback* AutoDetectCallback;
+	TUniquePtr<FAJAAutoDetectChannelCallback> AutoDetectCallback;
 };
diff --git a/Engine/Plugins/Media/AjaMedia/Source/AjaMedia/Public/AjaMediaSource.h b/Engine/Plugins/Media/AjaMedia/Source/AjaMedia/Public/AjaMediaSource.h
--- a/Engine/Plugins/Media/AjaMedia/Source/AjaMedia/Public/AjaMediaSource.h
+++ b/Engine/Plugins/Media/AjaMedia/Source/AjaMedia/Public/AjaMediaSource.h
@@ -44,10 +44,17 @@ public:
 	UPROPERTY(EditAnywhere, Category="AJA", meta=(DisplayName="Configuration"))
 	FMediaIOConfiguration MediaConfiguration;
 
+#if WITH_EDITORONLY_DATA
+	UE_DEPRECATED(5.1, "Use AutoDetectableTimecodeFormat")
 	/** Use the time code embedded in the input stream. */
-	UPROPERTY(BlueprintReadOnly, EditAnywhere, Category="AJA")
-	EMediaIOTimecodeFormat TimecodeFormat;
+	UPROPERTY(meta=(DeprecatedProperty))
+	EMediaIOTimecodeFormat TimecodeFormat_DEPRECATED = EMediaIOTimecodeFormat::None;
+#endif
 
+	/** Use the time code embedded in the input stream. */
+	UPROPERTY(BlueprintReadOnly, EditAnywhere, Category="AJA", meta = (DisplayName = "Timecode Format"))
+	EMediaIOAutoDetectableTimecodeFormat AutoDetectableTimecodeFormat = EMediaIOAutoDetectableTimecodeFormat::Auto; 
+	
 	/**
 	 * Use a ring buffer to capture and transfer data.
 	 * This may decrease transfer latency but increase stability.
@@ -135,6 +142,7 @@ public:
 	virtual void PostEditChangeChainProperty(struct FPropertyChangedChainEvent& PropertyChangedEvent) override;
 #endif //WITH_EDITOR
 
+	virtual void PostLoad() override;
 	virtual bool SupportsFormatAutoDetection() const override { return true; }
 	//~ End UObject interface
 };
diff --git a/Engine/Plugins/Media/AjaMedia/Source/AjaMediaEditor/Private/Customizations/AjaMediaSourceDetailsCustomization.cpp b/Engine/Plugins/Media/AjaMedia/Source/AjaMediaEditor/Private/Customizations/AjaMediaSourceDetailsCustomization.cpp
--- a/Engine/Plugins/Media/AjaMedia/Source/AjaMediaEditor/Private/Customizations/AjaMediaSourceDetailsCustomization.cpp
+++ b/Engine/Plugins/Media/AjaMedia/Source/AjaMediaEditor/Private/Customizations/AjaMediaSourceDetailsCustomization.cpp
@@ -103,7 +103,7 @@ void FAjaMediaSourceDetailsCustomization::OnAutoDetected(TArray<FAjaDeviceProvid
 			if (UAjaMediaSource* MediaSource = Cast<UAjaMediaSource>(EditingObject.Get()))
 			{
 				MediaSource->MediaConfiguration = Configurations[0].Configuration;
-				MediaSource->TimecodeFormat = Configurations[0].TimecodeFormat;
+				MediaSource->AutoDetectableTimecodeFormat = UE::MediaIO::ToAutoDetectableTimecodeFormat(Configurations[0].TimecodeFormat);
 			}
 		}
 	}
diff --git a/Engine/Plugins/Media/MediaIOFramework/Source/MediaIOCore/Public/MediaIOCoreDefinitions.h b/Engine/Plugins/Media/MediaIOFramework/Source/MediaIOCore/Public/MediaIOCoreDefinitions.h
--- a/Engine/Plugins/Media/MediaIOFramework/Source/MediaIOCore/Public/MediaIOCoreDefinitions.h
+++ b/Engine/Plugins/Media/MediaIOFramework/Source/MediaIOCore/Public/MediaIOCoreDefinitions.h
@@ -54,6 +54,63 @@ enum class EMediaIOTimecodeFormat
 	VITC,
 };
 
+/**
+ * Timecode formats.
+ */
+UENUM()
+enum class EMediaIOAutoDetectableTimecodeFormat : uint8
+{
+	None = 0,
+	LTC,
+	VITC,
+	Auto = 255,
+};
+
+namespace UE::MediaIO
+{
+	static EMediaIOAutoDetectableTimecodeFormat ToAutoDetectableTimecodeFormat(EMediaIOTimecodeFormat TimecodeFormat)
+	{
+		switch (TimecodeFormat)
+		{
+		case EMediaIOTimecodeFormat::None:
+			return EMediaIOAutoDetectableTimecodeFormat::None;
+		case EMediaIOTimecodeFormat::LTC:
+			return EMediaIOAutoDetectableTimecodeFormat::LTC;
+		case EMediaIOTimecodeFormat::VITC:
+			return EMediaIOAutoDetectableTimecodeFormat::VITC;
+		default:
+			ensureMsgf(false, TEXT("Unknown timecode format encountered!"));
+			return EMediaIOAutoDetectableTimecodeFormat::None;
+		}
+	}
+
+	static EMediaIOTimecodeFormat FromAutoDetectableTimecodeFormat(EMediaIOAutoDetectableTimecodeFormat TimecodeFormat)
+	{
+		switch (TimecodeFormat)
+		{
+		case EMediaIOAutoDetectableTimecodeFormat::LTC:
+			return EMediaIOTimecodeFormat::LTC;
+		case EMediaIOAutoDetectableTimecodeFormat::VITC:
+			return EMediaIOTimecodeFormat::VITC;
+		default:
+			ensureMsgf(false, TEXT("Unknown timecode format encountered!"));
+			return EMediaIOTimecodeFormat::None;
+		}
+	}
+}
+
+USTRUCT()
+struct MEDIAIOCORE_API FMediaIOAutoDetectableTimecodeFormat_Backup
+{
+	GENERATED_BODY()
+
+	/** The timecode format if not autodetected. */
+	UPROPERTY(EditAnywhere, Category = "AJA")
+	EMediaIOTimecodeFormat TimecodeFormat = EMediaIOTimecodeFormat::None;
+	/** Whether the timecode should be autodetected. */
+	UPROPERTY(EditAnywhere, Category = "AJA")
+	bool bAutoDetect = false;
+};
 
 /**
  * SDI Input type.
