Simplify aja timecode provider config selector

diff --git a/Engine/Plugins/Media/AjaMedia/Source/AjaMedia/Private/Assets/AjaTimecodeProvider.cpp b/Engine/Plugins/Media/AjaMedia/Source/AjaMedia/Private/Assets/AjaTimecodeProvider.cpp
--- a/Engine/Plugins/Media/AjaMedia/Source/AjaMedia/Private/Assets/AjaTimecodeProvider.cpp
+++ b/Engine/Plugins/Media/AjaMedia/Source/AjaMedia/Private/Assets/AjaTimecodeProvider.cpp
@@ -1,9 +1,10 @@
 // Copyright Epic Games, Inc. All Rights Reserved.
 
 #include "AjaTimecodeProvider.h"
-#include "AjaMediaPrivate.h"
+
 #include "AJA.h"
-
+#include "AjaDeviceProvider.h"
+#include "AjaMediaPrivate.h"
 #include "Misc/App.h"
 #include "UObject/VirtualProductionObjectVersion.h"
 
@@ -43,6 +44,7 @@
 #endif
 	, State(ETimecodeProviderSynchronizationState::Closed)
 {
+	DeviceProvider = MakePimpl<FAjaDeviceProvider>();
 }
 
 bool UAjaTimecodeProvider::FetchTimecode(FQualifiedFrameTime& OutFrameTime)
@@ -93,93 +95,23 @@
 		return false;
 	}
 
-	FString FailureReason;
-	if ((bUseDedicatedPin && !LTCConfiguration.IsValid()) || (!bUseDedicatedPin && !VideoConfiguration.IsValid()))
+	if (bAutoDetectTimecode)
 	{
-		State = ETimecodeProviderSynchronizationState::Error;
-		UE_LOG(LogAjaMedia, Warning, TEXT("The TimecodeProvider '%s' configuration is invalid."), *GetName());
-		return false;
+		DeviceProvider->AutoDetectConfiguration(FAjaDeviceProvider::FOnConfigurationAutoDetected::CreateUObject(this, &UAjaTimecodeProvider::OnConfigurationAutoDetected, InEngine));
+		return true;
 	}
-
-	check(SyncCallback == nullptr);
-	SyncCallback = new FAJACallback(this);
-
-	const int32 DeviceIndex = bUseDedicatedPin ? LTCConfiguration.Device.DeviceIdentifier : VideoConfiguration.MediaConfiguration.MediaConnection.Device.DeviceIdentifier;
-	AJA::AJADeviceOptions DeviceOptions(DeviceIndex);
-
-	AJA::AJATimecodeChannelOptions Options(*GetName());
-	Options.CallbackInterface = SyncCallback;
-	Options.bUseDedicatedPin = bUseDedicatedPin;
-
-	if (bUseDedicatedPin)
-	{
-		Options.bReadTimecodeFromReferenceIn = bUseReferenceIn;
-		Options.LTCSourceIndex = LTCConfiguration.LtcIndex;
-		Options.LTCFrameRateNumerator = LTCConfiguration.LtcFrameRate.Numerator;
-		Options.LTCFrameRateDenominator = LTCConfiguration.LtcFrameRate.Denominator;
-	}
 	else
 	{
-		Options.bReadTimecodeFromReferenceIn = false;
-		Options.LTCSourceIndex = 0;
-
-		Options.ChannelIndex = VideoConfiguration.MediaConfiguration.MediaConnection.PortIdentifier;
-		Options.VideoFormatIndex = VideoConfiguration.MediaConfiguration.MediaMode.DeviceModeIdentifier;
-
-		Options.TransportType = AJA::ETransportType::TT_SdiSingle;
+		FString FailureReason;
+		if ((bUseDedicatedPin && !LTCConfiguration.IsValid()) || (!bUseDedicatedPin && !VideoConfiguration.IsValid()))
 		{
-			const EMediaIOTransportType TransportType = VideoConfiguration.MediaConfiguration.MediaConnection.TransportType;
-			const EMediaIOQuadLinkTransportType QuadTransportType = VideoConfiguration.MediaConfiguration.MediaConnection.QuadTransportType;
-			switch (TransportType)
-			{
-			case EMediaIOTransportType::SingleLink:
-				Options.TransportType = AJA::ETransportType::TT_SdiSingle;
-				break;
-			case EMediaIOTransportType::DualLink:
-				Options.TransportType = AJA::ETransportType::TT_SdiDual;
-				break;
-			case EMediaIOTransportType::QuadLink:
-				Options.TransportType = QuadTransportType == EMediaIOQuadLinkTransportType::SquareDivision ? AJA::ETransportType::TT_SdiQuadSQ : AJA::ETransportType::TT_SdiQuadTSI;
-				break;
-			case EMediaIOTransportType::HDMI:
-				Options.TransportType = AJA::ETransportType::TT_Hdmi;
-				break;
-			}
-		}
-
-		Options.TimecodeFormat = AJA::ETimecodeFormat::TCF_None;
-		switch (VideoConfiguration.TimecodeFormat)
-		{
-		case EMediaIOTimecodeFormat::None:
-			Options.TimecodeFormat = AJA::ETimecodeFormat::TCF_None;
-			break;
-		case EMediaIOTimecodeFormat::LTC:
-			Options.TimecodeFormat = AJA::ETimecodeFormat::TCF_LTC;
-			break;
-		case EMediaIOTimecodeFormat::VITC:
-			Options.TimecodeFormat = AJA::ETimecodeFormat::TCF_VITC1;
-			break;
-		default:
-			break;
+			State = ETimecodeProviderSynchronizationState::Error;
+			UE_LOG(LogAjaMedia, Warning, TEXT("The TimecodeProvider '%s' configuration is invalid."), *GetName());
+			return false;
 		}
 	}
 
-	check(TimecodeChannel == nullptr);
-	TimecodeChannel = new AJA::AJATimecodeChannel();
-
-	if (!TimecodeChannel->Initialize(DeviceOptions, Options))
-	{
-		State = ETimecodeProviderSynchronizationState::Error;
-		ReleaseResources();
-		return false;
-	}
-
-#if WITH_EDITORONLY_DATA
-	InitializedEngine = InEngine;
-#endif
-
-	State = ETimecodeProviderSynchronizationState::Synchronizing;
-	return true;
+	return Initialize_Internal(InEngine, AJA::AJATimecodeChannelOptions(*GetName()));
 }
 
 void UAjaTimecodeProvider::Shutdown(class UEngine* InEngine)
@@ -223,6 +155,11 @@
 		delete SyncCallback;
 		SyncCallback = nullptr;
 	}
+
+	if (DeviceProvider)
+	{
+		DeviceProvider->EndAutoDetectConfiguration();
+	}
 }
 
 ETickableTickType UAjaTimecodeProvider::GetTickableTickType() const
@@ -259,4 +196,114 @@
 #endif
 }
 
+void UAjaTimecodeProvider::OnConfigurationAutoDetected(TArray<FAjaDeviceProvider::FMediaIOConfigurationWithTimecodeFormat> InConfigurations, class UEngine* InEngine)
+{
+	AJA::AJATimecodeChannelOptions Options(*GetName());
+
+	if (InConfigurations.Num())
+	{
+		FAjaDeviceProvider::FMediaIOConfigurationWithTimecodeFormat Format = InConfigurations[0];
+
+		VideoConfiguration.MediaConfiguration = Format.Configuration;
+		VideoConfiguration.TimecodeFormat = Format.TimecodeFormat;
+	}
+	else
+	{
+		UE_LOG(LogAjaMedia, Warning, TEXT("No configuration was detected for TimecodeProvider '%s'"), *GetName());
+		ReleaseResources();
+		return;
+	}
+
+	Initialize_Internal(InEngine, Options);
+}
+
+bool UAjaTimecodeProvider::Initialize_Internal(class UEngine* InEngine, AJA::AJATimecodeChannelOptions InOptions)
+{
+	check(SyncCallback == nullptr);
+	SyncCallback = new FAJACallback(this);
+
+	const int32 DeviceIndex = bUseDedicatedPin ? LTCConfiguration.Device.DeviceIdentifier : VideoConfiguration.MediaConfiguration.MediaConnection.Device.DeviceIdentifier;
+	if (DeviceIndex == -1)
+	{ 
+		State = ETimecodeProviderSynchronizationState::Error;
+		ReleaseResources();
+		return false;
+	}
+
+	AJA::AJADeviceOptions DeviceOptions(DeviceIndex);
+
+	InOptions.CallbackInterface = SyncCallback;
+	InOptions.bUseDedicatedPin = bUseDedicatedPin;
+
+	if (bUseDedicatedPin)
+	{
+		InOptions.bReadTimecodeFromReferenceIn = bUseReferenceIn;
+		InOptions.LTCSourceIndex = LTCConfiguration.LtcIndex;
+		InOptions.LTCFrameRateNumerator = LTCConfiguration.LtcFrameRate.Numerator;
+		InOptions.LTCFrameRateDenominator = LTCConfiguration.LtcFrameRate.Denominator;
+	}
+	else
+	{
+		InOptions.bReadTimecodeFromReferenceIn = false;
+		InOptions.LTCSourceIndex = 0;
+
+		InOptions.ChannelIndex = VideoConfiguration.MediaConfiguration.MediaConnection.PortIdentifier;
+		InOptions.VideoFormatIndex = VideoConfiguration.MediaConfiguration.MediaMode.DeviceModeIdentifier;
+
+		InOptions.TransportType = AJA::ETransportType::TT_SdiSingle;
+		{
+			const EMediaIOTransportType TransportType = VideoConfiguration.MediaConfiguration.MediaConnection.TransportType;
+			const EMediaIOQuadLinkTransportType QuadTransportType = VideoConfiguration.MediaConfiguration.MediaConnection.QuadTransportType;
+			switch (TransportType)
+			{
+			case EMediaIOTransportType::SingleLink:
+				InOptions.TransportType = AJA::ETransportType::TT_SdiSingle;
+				break;
+			case EMediaIOTransportType::DualLink:
+				InOptions.TransportType = AJA::ETransportType::TT_SdiDual;
+				break;
+			case EMediaIOTransportType::QuadLink:
+				InOptions.TransportType = QuadTransportType == EMediaIOQuadLinkTransportType::SquareDivision ? AJA::ETransportType::TT_SdiQuadSQ : AJA::ETransportType::TT_SdiQuadTSI;
+				break;
+			case EMediaIOTransportType::HDMI:
+				InOptions.TransportType = AJA::ETransportType::TT_Hdmi;
+				break;
+			}
+		}
+
+		InOptions.TimecodeFormat = AJA::ETimecodeFormat::TCF_None;
+		switch (VideoConfiguration.TimecodeFormat)
+		{
+		case EMediaIOTimecodeFormat::None:
+			InOptions.TimecodeFormat = AJA::ETimecodeFormat::TCF_None;
+			break;
+		case EMediaIOTimecodeFormat::LTC:
+			InOptions.TimecodeFormat = AJA::ETimecodeFormat::TCF_LTC;
+			break;
+		case EMediaIOTimecodeFormat::VITC:
+			InOptions.TimecodeFormat = AJA::ETimecodeFormat::TCF_VITC1;
+			break;
+		default:
+			break;
+		}
+	}
+
+	check(TimecodeChannel == nullptr);
+	TimecodeChannel = new AJA::AJATimecodeChannel();
+
+	if (!TimecodeChannel->Initialize(DeviceOptions, InOptions))
+	{
+		State = ETimecodeProviderSynchronizationState::Error;
+		ReleaseResources();
+		return false;
+	}
+
+#if WITH_EDITORONLY_DATA
+	InitializedEngine = InEngine;
+#endif
+
+	State = ETimecodeProviderSynchronizationState::Synchronizing;
+	return true;
+}
+
 #undef LOCTEXT_NAMESPACE
diff --git a/Engine/Plugins/Media/AjaMedia/Source/AjaMedia/Public/AjaTimecodeProvider.h b/Engine/Plugins/Media/AjaMedia/Source/AjaMedia/Public/AjaTimecodeProvider.h
--- a/Engine/Plugins/Media/AjaMedia/Source/AjaMedia/Public/AjaTimecodeProvider.h
+++ b/Engine/Plugins/Media/AjaMedia/Source/AjaMedia/Public/AjaTimecodeProvider.h
@@ -13,6 +13,7 @@
 namespace AJA
 {
 	class AJATimecodeChannel;
+	struct AJATimecodeChannelOptions;
 }
 
 class UEngine;
@@ -50,6 +51,8 @@ private:
 	friend FAJACallback;
 
 	void ReleaseResources();
+	bool Initialize_Internal(class UEngine* InEngine, AJA::AJATimecodeChannelOptions InOptions);
+	void OnConfigurationAutoDetected(TArray<FAjaDeviceProvider::FMediaIOConfigurationWithTimecodeFormat> InConfigurations, class UEngine* InEngine);
 
 public:
 
@@ -77,6 +80,9 @@ public:
 	UPROPERTY(EditAnywhere, Category="Timecode", meta=(EditCondition="!bUseDedicatedPin"))
 	FAjaMediaTimecodeConfiguration VideoConfiguration;
 
+	UPROPERTY()
+	bool bAutoDetectTimecode = true;
+
 private:
 	/** AJA channel associated with reading LTC timecode */
 	AJA::AJATimecodeChannel* TimecodeChannel;
@@ -93,4 +99,6 @@ private:
 
 	/** The current SynchronizationState of the TimecodeProvider*/
 	ETimecodeProviderSynchronizationState State;
+
+	TPimplPtr<class FAjaDeviceProvider> DeviceProvider;
 };
diff --git a/Engine/Plugins/Media/AjaMedia/Source/AjaMediaEditor/Private/Customizations/AjaMediaTimecodeConfigurationCustomization.cpp b/Engine/Plugins/Media/AjaMedia/Source/AjaMediaEditor/Private/Customizations/AjaMediaTimecodeConfigurationCustomization.cpp
--- a/Engine/Plugins/Media/AjaMedia/Source/AjaMediaEditor/Private/Customizations/AjaMediaTimecodeConfigurationCustomization.cpp
+++ b/Engine/Plugins/Media/AjaMedia/Source/AjaMediaEditor/Private/Customizations/AjaMediaTimecodeConfigurationCustomization.cpp
@@ -2,12 +2,14 @@
 
 #include "Customizations/AjaMediaTimecodeConfigurationCustomization.h"
 
+#include "AjaTimecodeProvider.h"
 #include "DetailWidgetRow.h"
 #include "DetailCategoryBuilder.h"
 #include "DetailLayoutBuilder.h"
 #include "Framework/Application/SlateApplication.h"
 #include "MediaIOPermutationsSelectorBuilder.h"
 #include "ObjectEditorUtils.h"
+#include "ScopedTransaction.h"
 #include "Widgets/Input/SComboButton.h"
 #include "Widgets/Text/STextBlock.h"
 #include "Widgets/SBoxPanel.h"
@@ -94,11 +96,35 @@ TSharedRef<SWidget> FAjaMediaTimecodeConfigurationCustomization::HandleSourceCom
 		return false;
 	};
 
+	TSharedRef<SWidget> EnforceWidget = SNew(SHorizontalBox)
+		+ SHorizontalBox::Slot()
+		.Padding(4.f)
+		.AutoWidth()
+		[
+			SNew(STextBlock)
+			.Text(LOCTEXT("EnforceLabel", "Enforce format"))
+		]
+	+ SHorizontalBox::Slot()
+		.Padding(4.f)
+		.AutoWidth()
+		[
+			SNew(SCheckBox)
+			.ToolTipText(LOCTEXT("EnforceTooltip", "Enforce a specific configuration. (Note that this disables format autodetection)"))
+			.IsChecked(this, &FAjaMediaTimecodeConfigurationCustomization::GetEnforceCheckboxState)
+			.OnCheckStateChanged(this, &FAjaMediaTimecodeConfigurationCustomization::SetEnforceCheckboxState)
+		];
+
+	auto GetExtensions = [EnforceWidget](TArray<TSharedRef<SWidget>>& OutWidgets)
+	{
+		OutWidgets.Add(EnforceWidget);
+	};
+
 	using TSelection = SMediaPermutationsSelector<FAjaMediaTimecodeConfiguration, AjaMediaTimecodeConfigurationCustomization::FMediaTimecodePermutationsSelectorBuilder>;
 	TSelection::FArguments Arguments;
 	Arguments
 		.PermutationsSource(MoveTemp(MediaConfigurations))
 		.SelectedPermutation(SelectedConfiguration)
+		.OnGetExtensions_Lambda(GetExtensions)
 		.OnSelectionChanged(this, &FAjaMediaTimecodeConfigurationCustomization::OnSelectionChanged)
 		.OnButtonClicked(this, &FAjaMediaTimecodeConfigurationCustomization::OnButtonClicked)
 		+ TSelection::Column(FMediaIOPermutationsSelectorBuilder::NAME_DeviceIdentifier)
@@ -110,12 +136,16 @@ TSharedRef<SWidget> FAjaMediaTimecodeConfigurationCustomization::HandleSourceCom
 		.IsColumnVisible_Lambda(QuadTypeVisible)
 		+ TSelection::Column(FMediaIOPermutationsSelectorBuilder::NAME_Resolution)
 		.Label(LOCTEXT("ResolutionLabel", "Resolution"))
+		.IsColumnVisible_Raw(this, &FAjaMediaTimecodeConfigurationCustomization::ShowAdvancedColumns)
 		+ TSelection::Column(FMediaIOPermutationsSelectorBuilder::NAME_Standard)
 		.Label(LOCTEXT("StandardLabel", "Standard"))
+		.IsColumnVisible_Raw(this, &FAjaMediaTimecodeConfigurationCustomization::ShowAdvancedColumns)
 		+ TSelection::Column(FMediaIOPermutationsSelectorBuilder::NAME_FrameRate)
 		.Label(LOCTEXT("FrameRateLabel", "Frame Rate"))
+		.IsColumnVisible_Raw(this, &FAjaMediaTimecodeConfigurationCustomization::ShowAdvancedColumns)
 		+ TSelection::Column(AjaMediaTimecodeConfigurationCustomization::NAME_TimecodeFormat)
-		.Label(LOCTEXT("TimecodeFormatLabel", "Format"));
+		.Label(LOCTEXT("TimecodeFormatLabel", "Format"))
+		.IsColumnVisible_Raw(this, &FAjaMediaTimecodeConfigurationCustomization::ShowAdvancedColumns);
 
 	TSharedRef<TSelection> Selector = SNew(TSelection) = Arguments;
 	PermutationSelector = Selector;
@@ -143,4 +173,58 @@ FReply FAjaMediaTimecodeConfigurationCustomization::OnButtonClicked() const
 	return FReply::Handled();
 }
 
+ECheckBoxState FAjaMediaTimecodeConfigurationCustomization::GetEnforceCheckboxState() const
+{
+	for (UObject* CustomizedObject : GetCustomizedObjects())
+	{
+		if (UAjaTimecodeProvider* Provider = Cast<UAjaTimecodeProvider>(CustomizedObject))
+		{
+			if (!Provider->bAutoDetectTimecode)
+			{
+				return ECheckBoxState::Checked;
+			}
+		}
+		else
+		{
+			return ECheckBoxState::Undetermined;
+		}
+	}
+
+	return ECheckBoxState::Unchecked;
+}
+
+void FAjaMediaTimecodeConfigurationCustomization::SetEnforceCheckboxState(ECheckBoxState CheckboxState)
+{
+	bool bEnforceFormat = CheckboxState == ECheckBoxState::Checked;
+
+	FScopedTransaction Transaction{ LOCTEXT("MediaIOAutoDetectTimecode", "Auto Detect Timecode") };
+	for (UObject* CustomizedObject : GetCustomizedObjects())
+	{
+		if (UAjaTimecodeProvider* Provider = Cast<UAjaTimecodeProvider>(CustomizedObject))
+		{
+			Provider->Modify();
+			Provider->bAutoDetectTimecode = !bEnforceFormat;
+		}
+	}
+}
+
+bool FAjaMediaTimecodeConfigurationCustomization::ShowAdvancedColumns(FName ColumnName, const TArray<FAjaMediaTimecodeConfiguration>& UniquePermutationsForThisColumn) const
+{
+	bool bAutoDetectTimecode = true;
+
+	for (UObject* CustomizedObject : GetCustomizedObjects())
+	{
+		if (UAjaTimecodeProvider* Provider = Cast<UAjaTimecodeProvider>(CustomizedObject))
+		{
+			if (!Provider->bAutoDetectTimecode)
+			{
+				bAutoDetectTimecode = false;
+				break;
+			}
+		}
+	}
+
+	return !bAutoDetectTimecode;
+}
+
 #undef LOCTEXT_NAMESPACE
diff --git a/Engine/Plugins/Media/AjaMedia/Source/AjaMediaEditor/Private/Customizations/AjaMediaTimecodeConfigurationCustomization.h b/Engine/Plugins/Media/AjaMedia/Source/AjaMediaEditor/Private/Customizations/AjaMediaTimecodeConfigurationCustomization.h
--- a/Engine/Plugins/Media/AjaMedia/Source/AjaMediaEditor/Private/Customizations/AjaMediaTimecodeConfigurationCustomization.h
+++ b/Engine/Plugins/Media/AjaMedia/Source/AjaMediaEditor/Private/Customizations/AjaMediaTimecodeConfigurationCustomization.h
@@ -21,7 +21,11 @@ private:
 
 	void OnSelectionChanged(FAjaMediaTimecodeConfiguration SelectedItem);
 	FReply OnButtonClicked() const;
+	ECheckBoxState GetEnforceCheckboxState() const;
+	void SetEnforceCheckboxState(ECheckBoxState CheckboxState);
+	bool ShowAdvancedColumns(FName ColumnName, const TArray<FAjaMediaTimecodeConfiguration>& UniquePermutationsForThisColumn) const;
 
+private:
 	TWeakPtr<SWidget> PermutationSelector;
 	FAjaMediaTimecodeConfiguration SelectedConfiguration;
 };
diff --git a/Engine/Plugins/Media/MediaIOFramework/Source/MediaIOEditor/Private/Customizations/MediaIOConfigurationCustomization.cpp b/Engine/Plugins/Media/MediaIOFramework/Source/MediaIOEditor/Private/Customizations/MediaIOConfigurationCustomization.cpp
--- a/Engine/Plugins/Media/MediaIOFramework/Source/MediaIOEditor/Private/Customizations/MediaIOConfigurationCustomization.cpp
+++ b/Engine/Plugins/Media/MediaIOFramework/Source/MediaIOEditor/Private/Customizations/MediaIOConfigurationCustomization.cpp
@@ -186,6 +186,11 @@ FReply FMediaIOConfigurationCustomization::OnButtonClicked() const
 
 bool FMediaIOConfigurationCustomization::ShowAdvancedColumns(FName ColumnName, const TArray<FMediaIOConfiguration>& UniquePermutationsForThisColumn) const
 {
+	if (!CustomizedSources.Num())
+	{
+		return true;
+	}
+
 	bool bAutoDetectInput = true;
 
 	for (const TWeakObjectPtr<UTimeSynchronizableMediaSource>& WeakMediaSource : CustomizedSources)
