Fix non-unity compile errors.

diff --git a/Engine/Plugins/Developer/PlasticSourceControl/Source/PlasticSourceControl/Private/PlasticSourceControlConsole.h b/Engine/Plugins/Developer/PlasticSourceControl/Source/PlasticSourceControl/Private/PlasticSourceControlConsole.h
--- a/Engine/Plugins/Developer/PlasticSourceControl/Source/PlasticSourceControl/Private/PlasticSourceControlConsole.h
+++ b/Engine/Plugins/Developer/PlasticSourceControl/Source/PlasticSourceControl/Private/PlasticSourceControlConsole.h
@@ -3,6 +3,7 @@
 #pragma once
 
 #include "CoreMinimal.h"
+#include "HAL/IConsoleManager.h"
 
 class FAutoConsoleCommand;
 
diff --git a/Engine/Plugins/Media/AjaMedia/Source/AjaMediaEditor/Private/Customizations/AjaMediaTimecodeConfigurationCustomization.cpp b/Engine/Plugins/Media/AjaMedia/Source/AjaMediaEditor/Private/Customizations/AjaMediaTimecodeConfigurationCustomization.cpp
--- a/Engine/Plugins/Media/AjaMedia/Source/AjaMediaEditor/Private/Customizations/AjaMediaTimecodeConfigurationCustomization.cpp
+++ b/Engine/Plugins/Media/AjaMedia/Source/AjaMediaEditor/Private/Customizations/AjaMediaTimecodeConfigurationCustomization.cpp
@@ -16,6 +16,7 @@
 #include "Widgets/SMediaPermutationsSelector.h"
 #include "Widgets/SNullWidget.h"
 #include "Widgets/SWindow.h"
+#include "Widgets/Input/SCheckBox.h"
 
 #define LOCTEXT_NAMESPACE "AjaMediaTimecodeConfigurationCustomization"
 
diff --git a/Engine/Plugins/Runtime/StateTree/Source/StateTreeModule/Private/StateTree.cpp b/Engine/Plugins/Runtime/StateTree/Source/StateTreeModule/Private/StateTree.cpp
--- a/Engine/Plugins/Runtime/StateTree/Source/StateTreeModule/Private/StateTree.cpp
+++ b/Engine/Plugins/Runtime/StateTree/Source/StateTreeModule/Private/StateTree.cpp
@@ -5,6 +5,7 @@
 #include "StateTreeNodeBase.h"
 #include "StateTreeTaskBase.h"
 #include "StateTreeEvaluatorBase.h"
+#include "AssetData.h"
 
 const FGuid FStateTreeCustomVersion::GUID(0x28E21331, 0x501F4723, 0x8110FA64, 0xEA10DA1E);
 FCustomVersionRegistration GRegisterStateTreeCustomVersion(FStateTreeCustomVersion::GUID, FStateTreeCustomVersion::LatestVersion, TEXT("StateTreeAsset"));
diff --git a/Engine/Plugins/Web/WebAPI/Source/WebAPIEditor/Public/CodeGen/Dom/WebAPICodeGenFile.h b/Engine/Plugins/Web/WebAPI/Source/WebAPIEditor/Public/CodeGen/Dom/WebAPICodeGenFile.h
--- a/Engine/Plugins/Web/WebAPI/Source/WebAPIEditor/Public/CodeGen/Dom/WebAPICodeGenFile.h
+++ b/Engine/Plugins/Web/WebAPI/Source/WebAPIEditor/Public/CodeGen/Dom/WebAPICodeGenFile.h
@@ -1,9 +1,10 @@
-﻿// Copyright Epic Games, Inc. All Rights Reserved.
+// Copyright Epic Games, Inc. All Rights Reserved.
 
 #pragma once
 
 #include "CoreMinimal.h"
 #include "WebAPICodeGenBase.h"
+#include "UObject/WeakObjectPtr.h"
 #include "UObject/WeakObjectPtrTemplates.h"
 
 class UWebAPIProperty;
