Fix aja timecode provider not working with dedicated pin

diff --git a/Engine/Plugins/Media/AjaMedia/Source/AjaMedia/Private/Assets/AjaTimecodeProvider.cpp b/Engine/Plugins/Media/AjaMedia/Source/AjaMedia/Private/Assets/AjaTimecodeProvider.cpp
--- a/Engine/Plugins/Media/AjaMedia/Source/AjaMedia/Private/Assets/AjaTimecodeProvider.cpp
+++ b/Engine/Plugins/Media/AjaMedia/Source/AjaMedia/Private/Assets/AjaTimecodeProvider.cpp
@@ -95,7 +95,7 @@ bool UAjaTimecodeProvider::Initialize(class UEngine* InEngine)
 		return false;
 	}
 
-	if (bAutoDetectTimecode)
+	if (bAutoDetectTimecode && !bUseDedicatedPin && !bUseReferenceIn)
 	{
 		DeviceProvider->AutoDetectConfiguration(FAjaDeviceProvider::FOnConfigurationAutoDetected::CreateUObject(this, &UAjaTimecodeProvider::OnConfigurationAutoDetected, InEngine));
 		return true;
