Add missing HDMI support for AJA outputs

diff --git a/Engine/Plugins/Media/AjaMedia/Source/AjaMedia/Private/Aja/AjaDeviceProvider.cpp b/Engine/Plugins/Media/AjaMedia/Source/AjaMedia/Private/Aja/AjaDeviceProvider.cpp
--- a/Engine/Plugins/Media/AjaMedia/Source/AjaMedia/Private/Aja/AjaDeviceProvider.cpp
+++ b/Engine/Plugins/Media/AjaMedia/Source/AjaMedia/Private/Aja/AjaDeviceProvider.cpp
@@ -498,7 +498,7 @@ TArray<FMediaIOConfiguration> FAjaDeviceProvider::GetConfigurations(bool bAllowI
 						}
 
 						const int32 HdmiPortCount = MediaConfiguration.bIsInput ? HdmiInputCount : HdmiOutputCount;
-						if (HdmiPortCount > 0 && MediaConfiguration.bIsInput) // only support HDMI input
+						if (HdmiPortCount > 0)
 						{
 							for (int32 FormatIndex = 0; FormatIndex < NumSupportedFormat; ++FormatIndex)
 							{
