Remove cvars for experimental interlace fix
Implement different fix for aja interlace field flipping

diff --git a/Engine/Plugins/Media/AjaMedia/Source/AjaMedia/Private/Assets/AjaCustomTimeStep.cpp b/Engine/Plugins/Media/AjaMedia/Source/AjaMedia/Private/Assets/AjaCustomTimeStep.cpp
--- a/Engine/Plugins/Media/AjaMedia/Source/AjaMedia/Private/Assets/AjaCustomTimeStep.cpp
+++ b/Engine/Plugins/Media/AjaMedia/Source/AjaMedia/Private/Assets/AjaCustomTimeStep.cpp
@@ -7,7 +7,6 @@
 #include "HAL/PlatformTime.h"
 #include "Misc/App.h"
 
-
 //~ IAJASyncChannelCallbackInterface implementation
 //--------------------------------------------------------------------
 // Those are called from the AJA thread. There's a lock inside AJA to prevent this object from dying while in this thread.
@@ -473,14 +472,36 @@ bool UAjaCustomTimeStep::WaitForSync()
 
 	const int32 ExpectedSyncCountsPerWait = GetExpectedSyncCountDelta();
 
-	if (bEnableOverrunDetection 
-		&& bIsNewSyncCountValid 
+	if (bIsNewSyncCountValid 
 		&& bIsPreviousSyncCountValid 
 		&& (NewSyncCount != (PreviousSyncCount + ExpectedSyncCountsPerWait)))
 	{
-		UE_LOG(LogAjaMedia, Warning, 
-			TEXT("The Engine couldn't run fast enough to keep up with the CustomTimeStep Sync. '%d' frame(s) dropped."), 
-			NewSyncCount - PreviousSyncCount - ExpectedSyncCountsPerWait);
+		uint32 Difference = NewSyncCount - PreviousSyncCount - ExpectedSyncCountsPerWait;
+		if (IsInterlaced() && UGenlockedCustomTimeStep::CVarExperimentalFieldFlipFix.GetValueOnAnyThread() && Difference % 2 != 0)
+		{
+			// Try not dropping an odd number of frames.
+			bool bSyncSuccess = true;
+			while (Difference % 2 != 0 && bSyncSuccess)
+			{
+				bSyncSuccess = SyncChannel->WaitForSync();
+				if (bSyncSuccess)
+				{
+					SyncChannel->GetSyncCount(NewSyncCount);
+					Difference = NewSyncCount - PreviousSyncCount - ExpectedSyncCountsPerWait;
+				}
+				else
+				{
+					UE_LOG(LogAjaMedia, Error, TEXT("The Engine couldn't run fast enough to keep up with the CustomTimeStep Sync. The wait timed out."));
+				}
+			}
+		}
+
+		if (bEnableOverrunDetection)
+		{
+			UE_LOG(LogAjaMedia, Warning, 
+				TEXT("The Engine couldn't run fast enough to keep up with the CustomTimeStep Sync. '%d' frame(s) dropped."), 
+				NewSyncCount - PreviousSyncCount - ExpectedSyncCountsPerWait);
+		}
 	}
 
 	if (bIsNewSyncCountValid)
@@ -530,4 +551,8 @@ void UAjaCustomTimeStep::DetectConfiguration(class UEngine* InEngine, bool bRein
 	}
 }
 
+bool UAjaCustomTimeStep::IsInterlaced() const
+{
+	return MediaConfiguration.MediaMode.Standard == EMediaIOStandardType::Interlaced;
+}
 
diff --git a/Engine/Plugins/Media/AjaMedia/Source/AjaMedia/Private/Player/AjaMediaPlayer.cpp b/Engine/Plugins/Media/AjaMedia/Source/AjaMedia/Private/Player/AjaMediaPlayer.cpp
--- a/Engine/Plugins/Media/AjaMedia/Source/AjaMedia/Private/Player/AjaMediaPlayer.cpp
+++ b/Engine/Plugins/Media/AjaMedia/Source/AjaMedia/Private/Player/AjaMediaPlayer.cpp
@@ -1,9 +1,18 @@
 // Copyright Epic Games, Inc. All Rights Reserved.
 
 #include "AjaMediaPlayer.h"
-#include "AjaMediaPrivate.h"
+
 
 #include "AJA.h"
+#include "AjaMediaAudioSample.h"
+#include "AjaMediaBinarySample.h"
+#include "AjaMediaPrivate.h"
+#include "AjaMediaSettings.h"
+#include "AjaMediaTextureSample.h"
+#include "Async/Async.h"
+
+#include "Engine/Engine.h"
+#include "GenlockedCustomTimeStep.h"
 #include "MediaIOCoreEncodeTime.h"
 #include "MediaIOCoreFileWriter.h"
 #include "MediaIOCoreSamples.h"
@@ -18,10 +27,6 @@
 #include "Stats/Stats2.h"
 #include "Styling/SlateStyle.h"
 
-#include "AjaMediaAudioSample.h"
-#include "AjaMediaBinarySample.h"
-#include "AjaMediaSettings.h"
-#include "AjaMediaTextureSample.h"
 
 #if WITH_EDITOR
 #include "EngineAnalytics.h"
@@ -95,6 +100,8 @@ FAjaMediaPlayer::~FAjaMediaPlayer()
  */
 bool FAjaMediaPlayer::Open(const FString& Url, const IMediaOptions* Options)
 {
+	bAddedSynchronizedSamples = false;
+
 	if (!FAja::CanUseAJACard())
 	{
 		UE_LOG(LogAjaMedia, Warning, TEXT("The AjaMediaPlayer can't open URL '%s' because Aja card cannot be used. Are you in a Commandlet? You may override this behavior by launching with -ForceAjaUsage"), *Url);
@@ -598,33 +605,50 @@ bool FAjaMediaPlayer::OnInputFrameReceived(const AJA::AJAInputFrameData& InInput
 		}
 		else
 		{
-			auto TextureSample = TextureSamplePool->AcquireShared();
-			if (InVideoFrame.bIsProgressivePicture)
+			bool bEvenField = GFrameCounter % 2 == 0;
+			if (bEvenField && !InVideoFrame.bIsProgressivePicture && !bAddedSynchronizedSamples && UGenlockedCustomTimeStep::CVarExperimentalFieldFlipFix.GetValueOnAnyThread())
 			{
-				if (TextureSample->InitializeProgressive(InVideoFrame, VideoSampleFormat, DecodedTime, VideoFrameRate, DecodedTimecode, bIsSRGBInput))
+				// Input must be received on an odd frame count or it flips the fields.
+				if (UGenlockedCustomTimeStep* TimeStep = Cast<UGenlockedCustomTimeStep>(GEngine->GetCustomTimeStep()))
 				{
-					Samples->AddVideo(TextureSample);
+					if (TimeStep->GetSynchronizationState() == ECustomTimeStepSynchronizationState::Synchronized)
+					{
+						// Reinitialize genlock if it it was started on the wrong frame (Must be started so the frames are received on odd frames.
+						AsyncTask(ENamedThreads::GameThread, [TimeStep]() {
+							TimeStep->Shutdown(GEngine);
+							TimeStep->Initialize(GEngine);
+							});
+					}
 				}
 			}
 			else
 			{
-				bool bEven = true;
+				auto TextureSample = TextureSamplePool->AcquireShared();
+				if (InVideoFrame.bIsProgressivePicture)
 
-				if (FMediaIOCorePlayerBase::CVarExperimentalFieldFlipFix.GetValueOnAnyThread())
 				{
-					bEven = GFrameCounterRenderThread % 2 != FMediaIOCorePlayerBase::CVarFlipInterlaceFields.GetValueOnAnyThread();
+					if (TextureSample->InitializeProgressive(InVideoFrame, VideoSampleFormat, DecodedTime, VideoFrameRate, DecodedTimecode, bIsSRGBInput))
+					{
+						Samples->AddVideo(TextureSample);
+					}
 				}
-
-				if (TextureSample->InitializeInterlaced_Halfed(InVideoFrame, VideoSampleFormat, DecodedTime, VideoFrameRate, DecodedTimecode, bEven, bIsSRGBInput))
-				{
-					Samples->AddVideo(TextureSample);
-				}
-
-				auto TextureSampleOdd = TextureSamplePool->AcquireShared();
-				bEven = !bEven;
-				if (TextureSampleOdd->InitializeInterlaced_Halfed(InVideoFrame, VideoSampleFormat, DecodedTimeF2, VideoFrameRate, DecodedTimecodeF2, bEven, bIsSRGBInput))
+				else
 				{
-					Samples->AddVideo(TextureSampleOdd);
+					bool bEven = true;
+					if (UGenlockedCustomTimeStep* TimeStep = Cast< UGenlockedCustomTimeStep>(GEngine->GetCustomTimeStep()))
+					{
+						bAddedSynchronizedSamples = TimeStep->GetSynchronizationState() == ECustomTimeStepSynchronizationState::Synchronized;
+					}
+					if (TextureSample->InitializeInterlaced_Halfed(InVideoFrame, VideoSampleFormat, DecodedTime, VideoFrameRate, DecodedTimecode, bEven, bIsSRGBInput))
+					{
+						Samples->AddVideo(TextureSample);
+					}
+					bEven = !bEven;
+					auto TextureSampleOdd = TextureSamplePool->AcquireShared();
+					if (TextureSampleOdd->InitializeInterlaced_Halfed(InVideoFrame, VideoSampleFormat, DecodedTimeF2, VideoFrameRate, DecodedTimecodeF2, bEven, bIsSRGBInput))
+					{
+						Samples->AddVideo(TextureSampleOdd);
+					}
 				}
 			}
 		}
diff --git a/Engine/Plugins/Media/AjaMedia/Source/AjaMedia/Private/Player/AjaMediaPlayer.h b/Engine/Plugins/Media/AjaMedia/Source/AjaMedia/Private/Player/AjaMediaPlayer.h
--- a/Engine/Plugins/Media/AjaMedia/Source/AjaMedia/Private/Player/AjaMediaPlayer.h
+++ b/Engine/Plugins/Media/AjaMedia/Source/AjaMedia/Private/Player/AjaMediaPlayer.h
@@ -197,4 +197,7 @@ private:
 
 	/** Autodetected or specified Timecode format. */
 	EMediaIOTimecodeFormat TimecodeFormat = EMediaIOTimecodeFormat::None; 
+
+	/** Whether samples have been added to the queue while the genlocked time step is synchronized. Reset when the time step is closed. */
+	bool bAddedSynchronizedSamples = false;
 };
diff --git a/Engine/Plugins/Media/AjaMedia/Source/AjaMedia/Public/AjaCustomTimeStep.h b/Engine/Plugins/Media/AjaMedia/Source/AjaMedia/Public/AjaCustomTimeStep.h
--- a/Engine/Plugins/Media/AjaMedia/Source/AjaMedia/Public/AjaCustomTimeStep.h
+++ b/Engine/Plugins/Media/AjaMedia/Source/AjaMedia/Public/AjaCustomTimeStep.h
@@ -54,6 +54,7 @@ private:
 	bool Initialize_Internal(class UEngine* InEngine);
 	void OnConfigurationAutoDetected(TArray<FAjaDeviceProvider::FMediaIOConfigurationWithTimecodeFormat> InConfigurations, class UEngine* InEngine, bool bReinitialize);
 	void DetectConfiguration(class UEngine* InEngine, bool bReinitialize);
+	bool IsInterlaced() const;
 
 public:
 	/**
diff --git a/Engine/Plugins/Media/BlackmagicMedia/Source/BlackmagicMedia/Private/Player/BlackmagicMediaPlayer.cpp b/Engine/Plugins/Media/BlackmagicMedia/Source/BlackmagicMedia/Private/Player/BlackmagicMediaPlayer.cpp
--- a/Engine/Plugins/Media/BlackmagicMedia/Source/BlackmagicMedia/Private/Player/BlackmagicMediaPlayer.cpp
+++ b/Engine/Plugins/Media/BlackmagicMedia/Source/BlackmagicMedia/Private/Player/BlackmagicMediaPlayer.cpp
@@ -357,11 +357,6 @@ namespace BlackmagicMediaPlayerHelpers
 					{
 						bool bEven = true;
 
-						if (FMediaIOCorePlayerBase::CVarExperimentalFieldFlipFix.GetValueOnAnyThread())
-						{
-							bEven = GFrameCounterRenderThread % 2 != FMediaIOCorePlayerBase::CVarFlipInterlaceFields.GetValueOnAnyThread();
-						}
-
 						auto TextureSampleEven = MediaPlayer->TextureSamplePool->AcquireShared();
 
 						if (TextureSampleEven->InitializeWithEvenOddLine(bEven
diff --git a/Engine/Plugins/Media/MediaIOFramework/Source/MediaIOCore/Private/MediaCapture.cpp b/Engine/Plugins/Media/MediaIOFramework/Source/MediaIOCore/Private/MediaCapture.cpp
--- a/Engine/Plugins/Media/MediaIOFramework/Source/MediaIOCore/Private/MediaCapture.cpp
+++ b/Engine/Plugins/Media/MediaIOFramework/Source/MediaIOCore/Private/MediaCapture.cpp
@@ -758,12 +758,6 @@ UMediaCapture::UMediaCapture(const FObjectInitializer& ObjectInitializer)
 
 UMediaCapture::~UMediaCapture() = default;
 
-UMediaCapture::UMediaCapture(FVTableHelper&)
-{
-
-}
-
-
 void UMediaCapture::BeginDestroy()
 {
 	if (GetState() == EMediaCaptureState::Capturing || GetState() == EMediaCaptureState::Preparing)
@@ -1467,7 +1461,7 @@ void UMediaCapture::InitializeOutputResources(int32 InNumberOfBuffers)
 	}
 
 	// Recreate frame manager which can trigger cleaning up its captured frames if it exists
-	FrameManager = MakeUnique<UE::MediaCaptureData::FFrameManager>();
+	FrameManager = MakePimpl<UE::MediaCaptureData::FFrameManager>();
 
 	bOutputResourcesInitialized = false;
 	NumberOfCaptureFrame = InNumberOfBuffers;
diff --git a/Engine/Plugins/Media/MediaIOFramework/Source/MediaIOCore/Private/Player/MediaIOCorePlayerBase.cpp b/Engine/Plugins/Media/MediaIOFramework/Source/MediaIOCore/Private/Player/MediaIOCorePlayerBase.cpp
--- a/Engine/Plugins/Media/MediaIOFramework/Source/MediaIOCore/Private/Player/MediaIOCorePlayerBase.cpp
+++ b/Engine/Plugins/Media/MediaIOFramework/Source/MediaIOCore/Private/Player/MediaIOCorePlayerBase.cpp
@@ -74,16 +74,6 @@ const FName FMediaIOCoreMediaOption::ResolutionWidth("ResolutionWidth");
 const FName FMediaIOCoreMediaOption::ResolutionHeight("ResolutionHeight");
 const FName FMediaIOCoreMediaOption::VideoModeName("VideoModeName");
 
-TAutoConsoleVariable<int32> FMediaIOCorePlayerBase::CVarFlipInterlaceFields(
-	TEXT("MediaIO.FlipInputInterlaceFields"), 0,
-	TEXT("Whether to flip input interlace fields. (Experimental)"),
-	ECVF_RenderThreadSafe);
-
-TAutoConsoleVariable<int32> FMediaIOCorePlayerBase::CVarExperimentalFieldFlipFix(
-	TEXT("MediaIO.PreventFieldFlipping"), 1,
-	TEXT("Whether to attempt fixing field flipping on input. (Experimental)"),
-	ECVF_RenderThreadSafe);
-
 /* FMediaIOCorePlayerBase structors
  *****************************************************************************/
 
diff --git a/Engine/Plugins/Media/MediaIOFramework/Source/MediaIOCore/Public/MediaCapture.h b/Engine/Plugins/Media/MediaIOFramework/Source/MediaIOCore/Public/MediaCapture.h
--- a/Engine/Plugins/Media/MediaIOFramework/Source/MediaIOCore/Public/MediaCapture.h
+++ b/Engine/Plugins/Media/MediaIOFramework/Source/MediaIOCore/Public/MediaCapture.h
@@ -16,6 +16,7 @@
 #include "RenderGraphResources.h"
 #include "RHI.h"
 #include "RHIResources.h"
+#include "Templates/PimplPtr.h"
 
 #include "MediaCapture.generated.h"
 
@@ -207,7 +208,6 @@ public:
 
 	/** Default constructor / destructor to use forward declared uniqueptr */
 	virtual ~UMediaCapture();
-	UMediaCapture(FVTableHelper&);
 
 	/**
 	 * Stop the actual capture if there is one.
@@ -491,7 +491,7 @@ private:
 	UPROPERTY(Transient)
 	TObjectPtr<UTextureRenderTarget2D> CapturingRenderTarget;
 	
-	TUniquePtr<UE::MediaCaptureData::FFrameManager> FrameManager;
+	TPimplPtr<UE::MediaCaptureData::FFrameManager> FrameManager;
 	int32 NumberOfCaptureFrame = 2;
 	int32 CaptureRequestCount = 0;
 	EMediaCaptureState MediaState = EMediaCaptureState::Stopped;
diff --git a/Engine/Plugins/Media/MediaIOFramework/Source/MediaIOCore/Public/MediaIOCorePlayerBase.h b/Engine/Plugins/Media/MediaIOFramework/Source/MediaIOCore/Public/MediaIOCorePlayerBase.h
--- a/Engine/Plugins/Media/MediaIOFramework/Source/MediaIOCore/Public/MediaIOCorePlayerBase.h
+++ b/Engine/Plugins/Media/MediaIOFramework/Source/MediaIOCore/Public/MediaIOCorePlayerBase.h
@@ -42,7 +42,7 @@ namespace UE::GPUTextureTransfer
 {
 	using TextureTransferPtr = TSharedPtr<class ITextureTransfer>;
 }
-
+ 
 /**
  * Implements a base player for hardware IO cards. 
  *
@@ -63,9 +63,6 @@ class MEDIAIOCORE_API FMediaIOCorePlayerBase
 	, public ITimedDataInput
 	, public TSharedFromThis<FMediaIOCorePlayerBase>
 {
-public:
-	static TAutoConsoleVariable<int32> CVarFlipInterlaceFields;
-	static TAutoConsoleVariable<int32> CVarExperimentalFieldFlipFix;
 public:
 
 	/**
diff --git a/Engine/Source/Runtime/TimeManagement/Private/GenlockedCustomTimeStep.cpp b/Engine/Source/Runtime/TimeManagement/Private/GenlockedCustomTimeStep.cpp
--- a/Engine/Source/Runtime/TimeManagement/Private/GenlockedCustomTimeStep.cpp
+++ b/Engine/Source/Runtime/TimeManagement/Private/GenlockedCustomTimeStep.cpp
@@ -3,6 +3,12 @@
 #include "GenlockedCustomTimeStep.h"
 #include "Misc/App.h"
 
+TAutoConsoleVariable<int32> UGenlockedCustomTimeStep::CVarExperimentalFieldFlipFix(
+	TEXT("MediaIO.PreventFieldFlipping"), 1,
+	TEXT("Whether to enable an interlace field flipping fix. (Experimental)"),
+	ECVF_RenderThreadSafe);
+
+
 void UGenlockedCustomTimeStep::UpdateAppTimes(const double& TimeBeforeSync, const double& TimeAfterSync) const
 {
 	// Use fixed delta time to update FApp times.
diff --git a/Engine/Source/Runtime/TimeManagement/Public/GenlockedCustomTimeStep.h b/Engine/Source/Runtime/TimeManagement/Public/GenlockedCustomTimeStep.h
--- a/Engine/Source/Runtime/TimeManagement/Public/GenlockedCustomTimeStep.h
+++ b/Engine/Source/Runtime/TimeManagement/Public/GenlockedCustomTimeStep.h
@@ -4,6 +4,7 @@
 
 #include "CoreTypes.h"
 #include "FixedFrameRateCustomTimeStep.h"
+#include "HAL/IConsoleManager.h"
 #include "Misc/AssertionMacros.h"
 #include "Misc/FrameRate.h"
 #include "UObject/ObjectMacros.h"
@@ -51,4 +52,7 @@ public:
 	/** Whether this custom time step should autodetect the video format if supported. */
 	UPROPERTY()
 	bool bAutoDetectFormat = false;
+
+	/** Experimental fix for interlace field flipping issue. */
+	static TAutoConsoleVariable<int32> CVarExperimentalFieldFlipFix;
 };
