Fix media capture outputting still image after engine closes unexpectedly

diff --git a/Engine/Plugins/Media/AjaMedia/Source/AjaMediaOutput/AjaMediaOutput.Build.cs b/Engine/Plugins/Media/AjaMedia/Source/AjaMediaOutput/AjaMediaOutput.Build.cs
--- a/Engine/Plugins/Media/AjaMedia/Source/AjaMediaOutput/AjaMediaOutput.Build.cs
+++ b/Engine/Plugins/Media/AjaMedia/Source/AjaMediaOutput/AjaMediaOutput.Build.cs
@@ -39,8 +39,14 @@ namespace UnrealBuildTool.Rules
 
 			if (Target.bBuildEditor == true)
 			{
-				PrivateDependencyModuleNames.Add("EditorFramework");
-				PrivateDependencyModuleNames.Add("UnrealEd");
+				PrivateDependencyModuleNames.AddRange(
+					new string[]
+					{
+						"EditorFramework",
+						"MainFrame",
+						"UnrealEd"
+					}
+				);
 			}
 		}
 	}
diff --git a/Engine/Plugins/Media/AjaMedia/Source/AjaMediaOutput/Private/AjaMediaCapture.cpp b/Engine/Plugins/Media/AjaMedia/Source/AjaMediaOutput/Private/AjaMediaCapture.cpp
--- a/Engine/Plugins/Media/AjaMedia/Source/AjaMediaOutput/Private/AjaMediaCapture.cpp
+++ b/Engine/Plugins/Media/AjaMedia/Source/AjaMediaOutput/Private/AjaMediaCapture.cpp
@@ -6,6 +6,7 @@
 #include "AjaDeviceProvider.h"
 #include "AjaMediaOutput.h"
 #include "AjaMediaOutputModule.h"
+#include "Misc/CoreDelegates.h"
 #include "Engine/Engine.h"
 #include "HAL/Event.h"
 #include "HAL/IConsoleManager.h"
@@ -20,6 +21,7 @@
 #if WITH_EDITOR
 #include "AnalyticsEventAttribute.h"
 #include "EngineAnalytics.h"
+#include "Interfaces/IMainFrameModule.h"
 #endif
 
 static TAutoConsoleVariable<int32> CVarAjaEnableGPUDirect(
@@ -222,6 +224,27 @@ UAjaMediaCapture::UAjaMediaCapture()
 	, WakeUpEvent(nullptr)
 {
 	bGPUTextureTransferAvailable = ((FAjaMediaOutputModule*)&IAjaMediaOutputModule::Get())->IsGPUTextureTransferAvailable();
+
+#if WITH_EDITOR
+	// In editor, an asset re-save dialog can prevent AJA from cleaning up in the regular PreExit callback,
+	// So we have to do our cleanup before the regular callback is called.
+	IMainFrameModule& MainFrame = FModuleManager::LoadModuleChecked<IMainFrameModule>("MainFrame");
+	CanCloseEditorDelegateHandle = MainFrame.RegisterCanCloseEditor(IMainFrameModule::FMainFrameCanCloseEditor::CreateUObject(this, &UAjaMediaCapture::CleanupPreEditorExit));
+#else
+	FCoreDelegates::OnEnginePreExit.AddUObject(this, &UAjaMediaCapture::OnEnginePreExit);
+#endif
+}
+
+UAjaMediaCapture::~UAjaMediaCapture()
+{
+#if WITH_EDITOR
+	if (IMainFrameModule* MainFrame = FModuleManager::GetModulePtr<IMainFrameModule>("MainFrame"))
+	{
+		MainFrame->UnregisterCanCloseEditor(CanCloseEditorDelegateHandle);
+	}
+#else
+	FCoreDelegates::OnEnginePreExit.RemoveAll(this);
+#endif
 }
 
 bool UAjaMediaCapture::ValidateMediaOutput() const
@@ -357,6 +380,23 @@ void UAjaMediaCapture::RestoreViewportTextureAlpha(TSharedPtr<FSceneViewport> In
 	}
 }
 
+bool UAjaMediaCapture::CleanupPreEditorExit()
+{
+	OnEnginePreExit();
+	return true;
+}
+
+void UAjaMediaCapture::OnEnginePreExit()
+{
+	if (OutputChannel)
+	{
+		// Close the aja channel in the another thread.
+		OutputChannel->Uninitialize();
+		OutputChannel.Reset();
+		OutputCallback.Reset();
+	}
+}
+
 bool UAjaMediaCapture::HasFinishedProcessing() const
 {
 	return Super::HasFinishedProcessing() || !OutputChannel;
@@ -468,7 +508,7 @@ bool UAjaMediaCapture::InitAJA(UAjaMediaOutput* InAjaMediaOutput)
 void UAjaMediaCapture::OnFrameCaptured_RenderingThread(const FCaptureBaseData& InBaseData, TSharedPtr<FMediaCaptureUserData, ESPMode::ThreadSafe> InUserData, void* InBuffer, int32 Width, int32 Height, int32 BytesPerRow)
 {
 	TRACE_CPUPROFILER_EVENT_SCOPE(UAjaMediaCapture::OnFrameCaptured_RenderingThread);
-	
+
 	// Prevent the rendering thread from copying while we are stopping the capture.
 	FScopeLock ScopeLock(&RenderThreadCriticalSection);
 	if (OutputChannel)
diff --git a/Engine/Plugins/Media/AjaMedia/Source/AjaMediaOutput/Public/AjaMediaCapture.h b/Engine/Plugins/Media/AjaMedia/Source/AjaMediaOutput/Public/AjaMediaCapture.h
--- a/Engine/Plugins/Media/AjaMedia/Source/AjaMediaOutput/Public/AjaMediaCapture.h
+++ b/Engine/Plugins/Media/AjaMedia/Source/AjaMediaOutput/Public/AjaMediaCapture.h
@@ -31,6 +31,7 @@ class AJAMEDIAOUTPUT_API UAjaMediaCapture : public UMediaCapture
 
 public:
 	UAjaMediaCapture();
+	~UAjaMediaCapture();
 
 	//~ UMediaCapture interface
 	virtual bool HasFinishedProcessing() const override;
@@ -59,6 +60,8 @@ private:
 	void OutputAudio_RenderingThread(const AJA::AJAOutputFrameBufferData& FrameBuffer) const;
 	void ApplyViewportTextureAlpha(TSharedPtr<FSceneViewport> InSceneViewport);
 	void RestoreViewportTextureAlpha(TSharedPtr<FSceneViewport> InSceneViewport);
+	bool CleanupPreEditorExit();
+	void OnEnginePreExit();
 
 	struct FAudioBuffer
 	{
@@ -104,4 +107,7 @@ private:
 
 	/** Whether or not GPUTextureTransfer was initialized successfully. */
 	bool bGPUTextureTransferAvailable = false;
+
+	/** Handle for the delegate used to clean up AJA on editor shutdown. */
+	FDelegateHandle CanCloseEditorDelegateHandle;
 };
