Fix static analysis errors in AjaCore

diff --git a/Engine/Plugins/Media/AjaMedia/Source/Aja/Private/Device.cpp b/Engine/Plugins/Media/AjaMedia/Source/Aja/Private/Device.cpp
--- a/Engine/Plugins/Media/AjaMedia/Source/Aja/Private/Device.cpp
+++ b/Engine/Plugins/Media/AjaMedia/Source/Aja/Private/Device.cpp
@@ -652,6 +652,7 @@ namespace AJA
 				if (Counter * SleepMilli > TimeoutMilli)
 				{
 					UE_LOG(LogTemp, Warning,  TEXT("Device: Can't get the device initialized.\n"));
+					delete Result;
 					return nullptr;
 				}
 				++Counter;
diff --git a/Engine/Plugins/Media/AjaMedia/Source/Aja/Private/DeviceScanner.cpp b/Engine/Plugins/Media/AjaMedia/Source/Aja/Private/DeviceScanner.cpp
--- a/Engine/Plugins/Media/AjaMedia/Source/Aja/Private/DeviceScanner.cpp
+++ b/Engine/Plugins/Media/AjaMedia/Source/Aja/Private/DeviceScanner.cpp
@@ -31,8 +31,6 @@ namespace AJA
 		bool DeviceScanner::GetDeviceTextId(int32_t InDeviceIndex, AJADeviceScanner::FormatedTextType& OutTextId) const
 		{
 			NTV2DeviceInfo DeviceInfo;
-			memset(&DeviceInfo, 0, sizeof(DeviceInfo));
-
 			if (!Scanner->GetDeviceInfo(InDeviceIndex, DeviceInfo, false))
 			{
 				OutTextId[0] = TEXT('\0');
diff --git a/Engine/Plugins/Media/AjaMedia/Source/Aja/Private/Helpers.cpp b/Engine/Plugins/Media/AjaMedia/Source/Aja/Private/Helpers.cpp
--- a/Engine/Plugins/Media/AjaMedia/Source/Aja/Private/Helpers.cpp
+++ b/Engine/Plugins/Media/AjaMedia/Source/Aja/Private/Helpers.cpp
@@ -487,15 +487,6 @@ namespace AJA
 				{
 					bool bDo4TSI = true;
 
-					if (InTransportType == ETransportType::TT_SdiSingle4kTSI)
-					{
-						bool bIs12Gb = false;
-						InCard->GetSDIInput12GPresent(bIs12Gb, InChannel);
-						bool bIs6Gb = false;
-						InCard->GetSDIInput6GPresent(bIs6Gb, InChannel);
-						bDo4TSI = bIs12Gb || bIs6Gb;
-					}
-
 					for (int32_t ChannelIndex = 0; ChannelIndex < NumberOfLinkChannel; ++ChannelIndex)
 					{
 						const NTV2Channel Channel = NTV2Channel(int32_t(InChannel) + ChannelIndex);
diff --git a/Engine/Plugins/Media/AjaMedia/Source/Aja/Private/Helpers.h b/Engine/Plugins/Media/AjaMedia/Source/Aja/Private/Helpers.h
--- a/Engine/Plugins/Media/AjaMedia/Source/Aja/Private/Helpers.h
+++ b/Engine/Plugins/Media/AjaMedia/Source/Aja/Private/Helpers.h
@@ -69,9 +69,9 @@ namespace AJA
 
 #ifdef NDEBUG
 //#define AJA_CHECK(FUNCTION) { if (!(FUNCTION)) { *reinterpret_cast<char*>(0) = 0; } }
-#define AJA_CHECK
+#define AJA_CHECK ensure
 #else
-#define AJA_CHECK assert
+#define AJA_CHECK check
 #endif
 
 #define AJA_TEST_MEMORY_BUFFER  0
