Try fixing stresstestunity build

diff --git a/Engine/Plugins/Media/AjaMedia/Source/Aja/AjaCore.Build.cs b/Engine/Plugins/Media/AjaMedia/Source/Aja/AjaCore.Build.cs
--- a/Engine/Plugins/Media/AjaMedia/Source/Aja/AjaCore.Build.cs
+++ b/Engine/Plugins/Media/AjaMedia/Source/Aja/AjaCore.Build.cs
@@ -39,7 +39,8 @@ public class AjaCore : ModuleRules
 
 			if (Target.Configuration == UnrealTargetConfiguration.Debug)
 			{
-				LibName = "libajantv2d.lib";
+				// Disable for now to fix clang build, to reproduce: compile in Debug Editor with -stresstestunity -nodebuginfo flags
+				//LibName = "libajantv2d.lib";
 			}
 
 			string LibPath = Path.Combine(AjaThirdPartyPath, "ntv2lib-deploy", "lib", LibName);
