Use bDisableStaticAnalysis for AjaCore

#rnx

diff --git a/Engine/Plugins/Media/AjaMedia/Source/Aja/AjaCore.Build.cs b/Engine/Plugins/Media/AjaMedia/Source/Aja/AjaCore.Build.cs
--- a/Engine/Plugins/Media/AjaMedia/Source/Aja/AjaCore.Build.cs
+++ b/Engine/Plugins/Media/AjaMedia/Source/Aja/AjaCore.Build.cs
@@ -31,8 +31,8 @@ public class AjaCore : ModuleRules
 				"RHI",
 			});
 
-		StaticAnalyzerDisabledCheckers.Add("all");
-		
+		bDisableStaticAnalysis = true;
+
 		if (Target.Platform == UnrealTargetPlatform.Win64)
 		{
 			string LibName = "libajantv2.lib";
