Fix audio delay in aja output, general improvements to audio handling in media io framework

diff --git a/Engine/Plugins/Media/AjaMedia/Source/Aja/AjaCore.Build.cs b/Engine/Plugins/Media/AjaMedia/Source/Aja/AjaCore.Build.cs
--- a/Engine/Plugins/Media/AjaMedia/Source/Aja/AjaCore.Build.cs
+++ b/Engine/Plugins/Media/AjaMedia/Source/Aja/AjaCore.Build.cs
@@ -29,6 +29,7 @@ public class AjaCore : ModuleRules
 			new string[] {
 				"Core",
 				"RHI",
+				"Engine"
 			});
 
 		bDisableStaticAnalysis = true;
diff --git a/Engine/Plugins/Media/AjaMedia/Source/Aja/Private/OutputChannel.cpp b/Engine/Plugins/Media/AjaMedia/Source/Aja/Private/OutputChannel.cpp
--- a/Engine/Plugins/Media/AjaMedia/Source/Aja/Private/OutputChannel.cpp
+++ b/Engine/Plugins/Media/AjaMedia/Source/Aja/Private/OutputChannel.cpp
@@ -534,7 +534,7 @@ namespace AJA
 
 			if (InAudioBufferSize > AudioBufferSize)
 			{
-				UE_LOG(LogTemp, Error, TEXT("SetAudioFrameData: Can't set the audio. The buffer is invalid or the buffer size is not the same as the AJA Audio Format for output channel %d on device %S.\n"), uint32_t(Channel) + 1, GetDevice().GetDisplayName().c_str());
+				UE_LOG(LogTemp, Error, TEXT("SetAudioFrameData: Can't set the audio. The buffer size is not the same as the AJA Audio Format for output channel %d on device %S.\n"), uint32_t(Channel) + 1, GetDevice().GetDisplayName().c_str());
 				return false;
 			}
 
@@ -548,16 +548,12 @@ namespace AJA
 			AJA_CHECK(AvailableWritingFrame->Audio[AudioBufferSize] == AJA_TEST_MEMORY_END_TAG);
 #endif
 
-			if (InAudioBufferSize == 0)
-			{
-				memset(AvailableWritingFrame->AudioBuffer, 0, AudioBufferSize);
-				InAudioBufferSize = AudioBufferSize;
-			}
-			else if(InAudioBuffer)
+			if (InAudioBuffer && InAudioBufferSize != 0)
 			{
 				memcpy(AvailableWritingFrame->AudioBuffer, InAudioBuffer, InAudioBufferSize);
 			}
 
+
 			AvailableWritingFrame->bAudioLineFilled = true;
 			AvailableWritingFrame->CopiedAudioBufferSize = InAudioBufferSize;
 			PushWhenFrameReady(AvailableWritingFrame);
@@ -751,7 +747,7 @@ namespace AJA
 				AJA_CHECK(Device->GetCard()->GetFrameRate(FrameRate, Channel));
 				NumSamplesPerFrame = ::GetAudioSamplesPerFrame(FrameRate, NTV2_AUDIO_48K);
 
-				if (AudioBufferSize > 0)
+				if (ensureMsgf(AudioBufferSize > 0, TEXT("Audio could not be initialized.")))
 				{
 					AJAAutoLock AutoLock(&FrameLock);
 					for (Frame* IttFrame : AllFrames)
@@ -766,11 +762,13 @@ namespace AJA
 						memset(IttFrame->AudioBuffer, 0, AudioBufferSize);
 					}
 
-					// Start audio output if a valid audio system is set (If UseAudio was true)
-					bResult &= GetDevice().StartAudioOutput(AudioSystem);
-					if (bResult == false)
+					if (GetOptions().bUseAutoCirculating)
 					{
-						UE_LOG(LogTemp, Error, TEXT("ConfigureAudio: Could not start audio output for audio system %d for channel %d on device %S."), uint32_t(AudioSystem), uint32_t(Channel) + 1, GetDevice().GetDisplayName().c_str());
+						bResult &= GetDevice().StartAudioOutput(AudioSystem);
+						if (bResult == false)
+						{
+							UE_LOG(LogTemp, Error, TEXT("ConfigureAudio: Could not start audio output for audio system %d for channel %d on device %S."), uint32_t(AudioSystem), uint32_t(Channel) + 1, GetDevice().GetDisplayName().c_str());
+						}
 					}
 				}
 			}
@@ -1173,63 +1171,9 @@ namespace AJA
 						}
 					}
 
-					if (UseAudio())
+					if (UseAudio() && AvailableReadingFrame->CopiedAudioBufferSize)
 					{
-						bool bAudioSuccess = true;
-						
-						ULWord AbsoluteAudioOffset = 0;
-						ULWord AudioWrapOffset = 0;
-						ULWord AudioLastOut = 0;
-
-						bAudioSuccess &= GetDevice().GetAudioMemoryOffset(CurrentAudioWriteOffset, AbsoluteAudioOffset, AudioSystem);
-						bAudioSuccess &= GetDevice().GetAudioWrapAddress(AudioWrapOffset);
-						bAudioSuccess &= GetDevice().ReadAudioLastOut(AudioLastOut);
-						
-						if (bAudioSuccess)
-						{
-							ULWord* AudioBuffer = reinterpret_cast<ULWord*>(AvailableReadingFrame->AudioBuffer);
-							if (CurrentAudioWriteOffset + AvailableReadingFrame->CopiedAudioBufferSize < AudioWrapOffset)
-							{
-								// Simplest case, just write the whole data to the SDRAM.
-								bRunning &= GetDevice().DMAWriteAudio(AudioSystem, AudioBuffer, CurrentAudioWriteOffset, AvailableReadingFrame->CopiedAudioBufferSize);
-								CurrentAudioWriteOffset += AvailableReadingFrame->CopiedAudioBufferSize;
-							}
-							else
-							{
-								// Audio data will wrap, so write in two parts.
-								ULWord BytesUntilBufferEnd = AudioBufferSize - CurrentAudioWriteOffset;
-								LWord RemainingBytes = AvailableReadingFrame->CopiedAudioBufferSize - BytesUntilBufferEnd;
-								// We might technically have space remaining, but we can't write there since it's too close to the wrap address.
-								if (RemainingBytes < 0)
-								{
-									RemainingBytes = 0;
-								}
-
-								bAudioSuccess &= GetDevice().DMAWriteAudio(AudioSystem, AudioBuffer, CurrentAudioWriteOffset, BytesUntilBufferEnd);
-								if (RemainingBytes >= 0)
-								{
-									bAudioSuccess &= GetDevice().DMAWriteAudio(AudioSystem, AudioBuffer + BytesUntilBufferEnd, 0, RemainingBytes);
-								}
-								CurrentAudioWriteOffset = RemainingBytes;
-							}
-
-							AudioLastOut = 0;
-							if (GetDevice().ReadAudioLastOut(AudioLastOut) && abs(LWord(CurrentAudioWriteOffset - AudioLastOut)) <= NumSamplesPerFrame)
-							{
-								UE_LOG(LogTemp, Warning,  TEXT("PingPong: Audio play head is catching up to write head for channel %d on device %S.\n"), uint32_t(Channel) + 1, GetDevice().GetDisplayName().c_str());
-							}
-
-							bool bAudioRunning = false;
-							if (GetDevice().IsAudioOutputRunning(AudioSystem, bAudioRunning) && !bAudioRunning)
-							{
-								UE_LOG(LogTemp, Warning,  TEXT("PingPong: Audio output option is true but the audio output is not running for channel %d on device %S.\n"), uint32_t(Channel) + 1, GetDevice().GetDisplayName().c_str());
-							}
-						}
-						else
-						{
-							bAudioSuccess = false;
-						}
-
+						bRunning &= Thread_HandleAudio(TEXT("PingPong"), AvailableReadingFrame);
 					}
 
 					if (!bRunning)
@@ -1274,6 +1218,12 @@ namespace AJA
 				}
 				else if (GetOptions().bDisplayWarningIfDropFrames)
 				{
+					if (UseAudio())
+					{
+						Thread_HandleLostFrameAudio();
+					}
+
+					TRACE_CPUPROFILER_EVENT_SCOPE(OutputChannel::NoFramesAvailable);
 					UE_LOG(LogTemp, Warning,  TEXT("PingPong: No frames are available for channel %d on device %S.\n"), uint32_t(Channel) + 1, GetDevice().GetDisplayName().c_str());
 				}
 			}
@@ -1340,5 +1290,193 @@ namespace AJA
 			}
 			++InterlacedTest_FrameCounter;
 		}
+
+		bool OutputChannelThread::Thread_GetAudioOffset(int32& OutAudioOffset)
+		{
+			bool bSuccess = true;
+			ULWord AudioWrapOffset = 0;
+			ULWord PlayHeadPosition = 0;
+
+			bSuccess &= GetDevice().GetAudioWrapAddress(AudioWrapOffset);
+			bSuccess &= GetDevice().ReadAudioLastOut(PlayHeadPosition);
+
+			// Calculate audio offset
+			if (CurrentAudioWriteOffset > PlayHeadPosition)
+			{
+				OutAudioOffset = CurrentAudioWriteOffset - PlayHeadPosition;
+			}
+			else
+			{
+				OutAudioOffset = (AudioWrapOffset - PlayHeadPosition) + CurrentAudioWriteOffset;
+			}
+
+			return bSuccess;
+		}
+
+
+		bool OutputChannelThread::Thread_GetAudioOffsetInSeconds(double& OutAudioOffset)
+		{
+			int32 OffsetInBytes = 0;
+			if (Thread_GetAudioOffset(OffsetInBytes))
+			{
+				OutAudioOffset = (OffsetInBytes / sizeof(int32)) / 48000.f / Options.NumberOfAudioChannel;
+				return true;
+			}
+
+			return false;
+		}
+
+		bool OutputChannelThread::Thread_GetAudioOffsetInSamples(int32& OutAudioOffset)
+		{
+			int32 OffsetInBytes = 0;
+			if (Thread_GetAudioOffset(OffsetInBytes))
+			{
+				OutAudioOffset = FMath::RoundToInt32((float)OffsetInBytes / Options.NumberOfAudioChannel / sizeof(int32));
+				return true;
+			}
+
+			return false;
+		}
+
+
+		bool OutputChannelThread::Thread_TransferAudioBuffer(uint8* InBuffer, int32 InBufferSize)
+		{
+			bool bSuccess = true;
+
+			ULWord AudioWrapOffset = 0;
+			ULWord WriteHeadPosition = 0;
+			ULWord PlayHeadPosition = 0;
+
+			bSuccess &= GetDevice().GetAudioWrapAddress(AudioWrapOffset);
+
+			if (CurrentAudioWriteOffset + InBufferSize < AudioWrapOffset)
+			{
+				// Simplest case, just write the whole data to the SDRAM.
+				bSuccess &= GetDevice().DMAWriteAudio(AudioSystem, reinterpret_cast<ULWord*>(InBuffer), CurrentAudioWriteOffset, InBufferSize);
+				CurrentAudioWriteOffset += InBufferSize;
+			}
+			else
+			{
+				// Audio data will wrap, so write in two parts.
+				ULWord BytesUntilBufferEnd = AudioWrapOffset - CurrentAudioWriteOffset;
+				LWord RemainingBytes = InBufferSize - BytesUntilBufferEnd;
+				// We might technically have space remaining, but we can't write there since it's too close to the wrap address.
+				if (RemainingBytes < 0)
+				{
+					RemainingBytes = 0;
+				}
+
+				bSuccess &= GetDevice().DMAWriteAudio(AudioSystem, reinterpret_cast<ULWord*>(InBuffer), CurrentAudioWriteOffset, BytesUntilBufferEnd);
+				if (RemainingBytes > 0)
+				{
+					bSuccess &= GetDevice().DMAWriteAudio(AudioSystem, reinterpret_cast<ULWord*>(InBuffer) + BytesUntilBufferEnd, 0, RemainingBytes);
+				}
+
+				CurrentAudioWriteOffset = RemainingBytes;
+			}
+
+			return bSuccess;
+		}
+
+		bool OutputChannelThread::Thread_HandleAudio(const FString& OutputMethod, Frame* AvailableReadingFrame)
+		{
+			bool bAudioTransferWasSuccessful = true;
+
+			bool bAudioSuccess = true;
+
+			bool bAudioOutputRunning = false;
+			bAudioSuccess &= GetDevice().IsAudioOutputRunning(AudioSystem, bAudioOutputRunning);
+
+			bool bIsPaused = false;
+			bAudioSuccess &= GetDevice().GetAudioOutputPause(AudioSystem, bIsPaused);
+
+			if (bIsPaused && bAudioOutputRunning && !Thread_ShouldPauseAudioOutput())
+			{
+				UE_LOG(LogTemp, Verbose, TEXT("%s: Resuming audio output."), *OutputMethod);
+				bool bPausePlayback = false;
+				bAudioSuccess &= GetDevice().SetAudioOutputPause(AudioSystem, bPausePlayback);
+				bIsPaused = true;
+			}
+			else if (!bAudioOutputRunning)
+			{
+				// Only offset the start ptr if this is the first time we're starting the audio output.
+
+				ULWord WrapOffset = 0;
+				bAudioSuccess &= GetDevice().GetAudioWrapAddress(WrapOffset);
+
+				int32 NumOffsetFrames = 1;
+				int32 InitialOffset = Options.NumberOfAudioChannel * NumSamplesPerFrame * NumOffsetFrames * sizeof(int32);
+
+				CurrentAudioWriteOffset = InitialOffset;
+
+				UE_LOG(LogTemp, Verbose, TEXT("%s: Starting audio output."), *OutputMethod);
+				bAudioSuccess &= GetDevice().StartAudioOutput(AudioSystem); // This resets the playhead to 0, we have to override that to something closer to the write head in order to reduce audio delay
+			}
+
+			if (bAudioSuccess)
+			{
+				double OffsetInSeconds = 0;
+				bAudioSuccess &= Thread_GetAudioOffsetInSeconds(OffsetInSeconds);
+				bAudioSuccess &= GetDevice().ReadAudioLastOut(AudioPlayheadLastPosition, AudioSystem);
+				UE_LOG(LogTemp, Verbose, TEXT("%s: Play head: %d, Write Head: %d, Offset: %f"), *OutputMethod, AudioPlayheadLastPosition, CurrentAudioWriteOffset, OffsetInSeconds);
+
+				bAudioTransferWasSuccessful &= Thread_TransferAudioBuffer(AvailableReadingFrame->AudioBuffer, AvailableReadingFrame->CopiedAudioBufferSize);
+
+				int32 OffsetInSamples = 0;
+				bAudioSuccess &= Thread_GetAudioOffsetInSamples(OffsetInSamples);
+
+				if (bAudioSuccess && Thread_ShouldPauseAudioOutput())
+				{
+					UE_LOG(LogTemp, Warning, TEXT("%s: Audio play head is catching up to write head for channel %d on device %S.\n"), *OutputMethod, uint32_t(Channel) + 1, GetDevice().GetDisplayName().c_str());
+
+					bool bPausePlayout = true;
+					GetDevice().SetAudioOutputPause(AudioSystem, bPausePlayout);
+
+					UE_LOG(LogTemp, Verbose, TEXT("%s Pausing audio output because the play head is catching up to the write head. Last playhead position: %d"), *OutputMethod, AudioPlayheadLastPosition);
+				}
+			}
+			else
+			{
+				UE_LOG(LogTemp, Verbose, TEXT("%s Audio output failed."), *OutputMethod);
+			}
+
+			return bAudioTransferWasSuccessful;
+		}
+
+		bool OutputChannelThread::Thread_HandleLostFrameAudio()
+		{
+			bool bAudioSuccess = true;
+			bool bAudioOutputRunning = false;
+			bAudioSuccess &= GetDevice().IsAudioOutputRunning(AudioSystem, bAudioOutputRunning);
+
+			if (bAudioOutputRunning)
+			{
+				// When dropping frames, we have to keep writing audio data or else the play head will catch up to the write head which will induce a delay
+				int32 OffsetInSeconds = 0;
+				bAudioSuccess &= Thread_GetAudioOffsetInSamples(OffsetInSeconds);
+
+				if (bAudioSuccess && Thread_ShouldPauseAudioOutput())
+				{
+					constexpr bool bPausePlayout = true;
+					bAudioSuccess &= GetDevice().SetAudioOutputPause(AudioSystem, bPausePlayout);
+					UE_LOG(LogTemp, Warning, TEXT("Pausing audio output. Last playhead position: %d"), AudioPlayheadLastPosition);
+				}
+			}
+
+			return bAudioSuccess;
+		}
+
+		bool OutputChannelThread::Thread_ShouldPauseAudioOutput()
+		{
+			int32 AudioBufferZoneInFrames = 1;
+			int32 OffsetInSamples = 0;
+			Thread_GetAudioOffsetInSamples(OffsetInSamples);
+
+			if (OffsetInSamples <= (int32) NumSamplesPerFrame * AudioBufferZoneInFrames)
+			{
+				return true;
+			}
+			return false;
+		}
 	}
 }
\ No newline at end of file
diff --git a/Engine/Plugins/Media/AjaMedia/Source/Aja/Private/OutputChannel.h b/Engine/Plugins/Media/AjaMedia/Source/Aja/Private/OutputChannel.h
--- a/Engine/Plugins/Media/AjaMedia/Source/Aja/Private/OutputChannel.h
+++ b/Engine/Plugins/Media/AjaMedia/Source/Aja/Private/OutputChannel.h
@@ -79,6 +79,15 @@ namespace AJA
 			void Thread_PushAvailableReadingFrame(Frame* InCurrentReadingFrame);
 			void ThreadLock_PingPongOutputLoop_Memset0(Frame* InCurrentReadingFrame);
 
+			/** Get the number of bytes separating the audio write and play head. */
+			bool Thread_GetAudioOffset(int32& OutAudioOffset);
+			bool Thread_GetAudioOffsetInSeconds(double& OutAudioOffset);
+			bool Thread_GetAudioOffsetInSamples(int32& OutAudioOffset);
+			bool Thread_TransferAudioBuffer(uint8* InBuffer, int32 InBufferSize);
+			bool Thread_HandleAudio(const FString& OutputMethod, Frame* AvailableReadingFrame);
+			bool Thread_HandleLostFrameAudio();
+			bool Thread_ShouldPauseAudioOutput();
+
 			struct Frame
 			{
 				Frame();
@@ -127,6 +136,7 @@ namespace AJA
 
 			ULWord CurrentAudioWriteOffset = 0;
 			int32_t NumSamplesPerFrame = 0;
+			ULWord AudioPlayheadLastPosition = 0;
 		};
 	}
 }
diff --git a/Engine/Plugins/Media/AjaMedia/Source/AjaMediaOutput/Private/AjaMediaCapture.cpp b/Engine/Plugins/Media/AjaMedia/Source/AjaMediaOutput/Private/AjaMediaCapture.cpp
--- a/Engine/Plugins/Media/AjaMedia/Source/AjaMediaOutput/Private/AjaMediaCapture.cpp
+++ b/Engine/Plugins/Media/AjaMedia/Source/AjaMediaOutput/Private/AjaMediaCapture.cpp
@@ -637,7 +637,9 @@ void UAjaMediaCapture::OutputAudio_RenderingThread(const AJA::AJAOutputFrameBuff
 	if (bOutputAudio)
 	{
 		TRACE_CPUPROFILER_EVENT_SCOPE(UAjaMediaCapture::OnFrameCaptured_RenderingThread::SetAudio);
-		TArray<int32> AudioSamples = AudioOutput->GetAudioSamples<int32>();
+
+		int32 NumSamplesToPull = FMath::RoundToInt32(48000.f * AudioOutput->NumInputChannels / FrameRate.AsDecimal());
+		TArray<int32> AudioSamples = AudioOutput->GetAudioSamples<int32>(NumSamplesToPull);
 		OutputChannel->SetAudioFrameData(FrameBuffer, reinterpret_cast<uint8*>(AudioSamples.GetData()), AudioSamples.Num() * sizeof(int32));
 	}
 }
diff --git a/Engine/Plugins/Media/AjaMedia/Source/AjaMediaOutput/Public/AjaMediaCapture.h b/Engine/Plugins/Media/AjaMedia/Source/AjaMediaOutput/Public/AjaMediaCapture.h
--- a/Engine/Plugins/Media/AjaMedia/Source/AjaMediaOutput/Public/AjaMediaCapture.h
+++ b/Engine/Plugins/Media/AjaMedia/Source/AjaMediaOutput/Public/AjaMediaCapture.h
@@ -110,4 +110,6 @@ private:
 
 	/** Handle for the delegate used to clean up AJA on editor shutdown. */
 	FDelegateHandle CanCloseEditorDelegateHandle;
+
+	mutable int32 AudioSamplesSentLastFrame = 0;
 };
diff --git a/Engine/Plugins/Media/BlackmagicMedia/Source/BlackmagicMediaOutput/Private/BlackmagicMediaCapture.cpp b/Engine/Plugins/Media/BlackmagicMedia/Source/BlackmagicMediaOutput/Private/BlackmagicMediaCapture.cpp
--- a/Engine/Plugins/Media/BlackmagicMedia/Source/BlackmagicMediaOutput/Private/BlackmagicMediaCapture.cpp
+++ b/Engine/Plugins/Media/BlackmagicMedia/Source/BlackmagicMediaOutput/Private/BlackmagicMediaCapture.cpp
@@ -760,7 +760,7 @@ void UBlackmagicMediaCapture::OutputAudio_RenderingThread(const FCaptureBaseData
 		
 		if (AudioBitDepth == EBlackmagicMediaOutputAudioBitDepth::Signed_32Bits)
 		{
-			TArray<int32> AudioBuffer = AudioOutput->GetAudioSamples<int32>(NumSamplesToPull);
+			TArray<int32> AudioBuffer = AudioOutput->GetAllAudioSamples<int32>();
 			AudioSamples.AudioBuffer = reinterpret_cast<uint8_t*>(AudioBuffer.GetData());
 			AudioSamples.NumAudioSamples = AudioBuffer.Num() / NumOutputChannels;
 			AudioSamples.AudioBufferLength = AudioBuffer.Num() * sizeof(int32);
@@ -768,7 +768,7 @@ void UBlackmagicMediaCapture::OutputAudio_RenderingThread(const FCaptureBaseData
 		}
 		else if (AudioBitDepth == EBlackmagicMediaOutputAudioBitDepth::Signed_16Bits)
 		{
-			TArray<int16> AudioBuffer = AudioOutput->GetAudioSamples<int16>(NumSamplesToPull);
+			TArray<int16> AudioBuffer = AudioOutput->GetAllAudioSamples<int16>();
 			AudioSamples.AudioBuffer = reinterpret_cast<uint8_t*>(AudioBuffer.GetData());
 			AudioSamples.NumAudioSamples = AudioBuffer.Num() / NumOutputChannels;
 			AudioSamples.AudioBufferLength = AudioBuffer.Num() * sizeof(int16);
diff --git a/Engine/Plugins/Media/MediaIOFramework/Source/MediaIOCore/MediaIOCore.Build.cs b/Engine/Plugins/Media/MediaIOFramework/Source/MediaIOCore/MediaIOCore.Build.cs
--- a/Engine/Plugins/Media/MediaIOFramework/Source/MediaIOCore/MediaIOCore.Build.cs
+++ b/Engine/Plugins/Media/MediaIOFramework/Source/MediaIOCore/MediaIOCore.Build.cs
@@ -48,6 +48,8 @@ namespace UnrealBuildTool.Rules
 	                "SoundFieldRendering"
 				});
 
+			PublicDefinitions.Add("WITH_MEDIA_IO_AUDIO_DEBUGGING=1");
+
 			if (Target.bBuildEditor == true)
 			{
 				PrivateDependencyModuleNames.Add("EditorFramework");
diff --git a/Engine/Plugins/Media/MediaIOFramework/Source/MediaIOCore/Private/MediaIOCoreAudioOutput.cpp b/Engine/Plugins/Media/MediaIOFramework/Source/MediaIOCore/Private/MediaIOCoreAudioOutput.cpp
--- a/Engine/Plugins/Media/MediaIOFramework/Source/MediaIOCore/Private/MediaIOCoreAudioOutput.cpp
+++ b/Engine/Plugins/Media/MediaIOFramework/Source/MediaIOCore/Private/MediaIOCoreAudioOutput.cpp
@@ -18,12 +18,12 @@ DEFINE_LOG_CATEGORY_STATIC(LogMediaIOAudioOutput, Log, All);
 #define LOCTEXT_NAMESPACE "MediaIOAudioOutput"
 
 FMediaIOAudioOutput::FMediaIOAudioOutput(Audio::FPatchOutputStrongPtr InPatchOutput, const FAudioOptions& InAudioOptions)
-	: PatchOutput(MoveTemp(InPatchOutput))
-    , NumInputChannels(InAudioOptions.InNumInputChannels)
+    : NumInputChannels(InAudioOptions.InNumInputChannels)
     , NumOutputChannels(InAudioOptions.InNumOutputChannels)
     , TargetFrameRate(InAudioOptions.InTargetFrameRate)
     , MaxSampleLatency(InAudioOptions.InMaxSampleLatency)
     , OutputSampleRate(InAudioOptions.InOutputSampleRate)
+	, PatchOutput(MoveTemp(InPatchOutput))
 {
 	NumSamplesPerFrame = FMath::CeilToInt(NumInputChannels * OutputSampleRate / TargetFrameRate.AsDecimal());
 }
@@ -49,6 +49,7 @@ Audio::FAlignedFloatBuffer FMediaIOAudioOutput::GetFloatBuffer(uint32 NumSamples
 	FloatBuffer.SetNumZeroed(NumSamplesToPop);
 
 	const int32 NumPopped = GetAudioBuffer(NumSamplesToPop, FloatBuffer.GetData());
+	FloatBuffer.SetNum(NumPopped);
 
 	return FloatBuffer;
 }
@@ -74,14 +75,17 @@ FMediaIOAudioCapture::~FMediaIOAudioCapture()
 
 void FMediaIOAudioCapture::OnNewSubmixBuffer(const USoundSubmix* InOwningSubmix, float* InAudioData, int32 InNumSamples, int32 InNumChannels, const int32 InSampleRate, double InAudioClock)
 {
+	TRACE_CPUPROFILER_EVENT_SCOPE(FMediaIOAudioCapture::OnNewSubmixBuffer);
+
 	check(InOwningSubmix);
 	if (InOwningSubmix->GetFName() == PrimarySubmixName)
 	{
 		if (ensureMsgf(NumChannels == InNumChannels, TEXT("Expected %d channels from submix buffer but got %d instead."), NumChannels, InNumChannels))
 		{
 			int32 NumPushed = AudioSplitter.PushAudio(InAudioData, InNumSamples);
-			if (InNumSamples != NumPushed)
+			if (InNumSamples != NumPushed && NumPushed != -1)
 			{
+				TRACE_CPUPROFILER_EVENT_SCOPE(FMediaIOAudioCapture::OnNewSubmixBuffer::BufferOverrun);
 				UE_LOG(LogMediaIOAudioOutput, Verbose, TEXT("Pushed samples mismatch, Incoming samples: %d, Pushed samples: %d"), InNumSamples, NumPushed);
 			}
 		}
diff --git a/Engine/Plugins/Media/MediaIOFramework/Source/MediaIOCore/Public/MediaIOAudioDebug.h b/Engine/Plugins/Media/MediaIOFramework/Source/MediaIOCore/Public/MediaIOAudioDebug.h
new file mode 100644
--- /dev/null
+++ b/Engine/Plugins/Media/MediaIOFramework/Source/MediaIOCore/Public/MediaIOAudioDebug.h
@@ -0,0 +1,415 @@
+// Copyright Epic Games, Inc. All Rights Reserved.
+
+#pragma once
+
+#include "CoreMinimal.h"
+#include "Engine/Engine.h"
+#include "HAL/FileManager.h"
+#include "HAL/IConsoleManager.h"
+#include "Logging/LogVerbosity.h"
+
+/**
+ * Single header drop in file to help diagnose audio issues
+ * 
+ * Usage:
+ * 1. Copy and paste this file where needed.
+ * 2. Gather a single channel from the audio data into an array
+ * 3. Add MEDIA_IO_DUMP_AUDIO(YourArray)
+ * 4. At runtime in the console command, enter MediaIO.DumpAudio <number_of_seconds>
+ * 
+ * This will generate a .csv and .wav file in your project's saved folder under the "AudioDumps" folder.
+ * 
+ * Current limitations: Only handles mono input at the moment, assumes 48000 sample rate.
+ */
+
+DEFINE_LOG_CATEGORY_STATIC(LogMediaIOAudioDebug, Log, All)
+
+#define MEDIA_IO_DUMP_AUDIO_MONO(Array)\
+	UE::MediaIOAudioDebug::GetSingleton().ProcessAudio(L#Array, Array);
+
+#define MEDIA_IO_DUMP_AUDIO(Name, Buffer, BufferSize, SampleSize, NumChannels)\
+	UE::MediaIOAudioDebug::GetSingleton().ProcessAudio(Name, Buffer, BufferSize, SampleSize, NumChannels);
+
+// Note: This was copied from audio.cpp and adapted to support multiple audio bit depth and floating point data.
+// This will eventually be removed when these changes are ported to the Audio.cpp file.
+namespace UE::MediaIOAudioDebug::Private
+{
+	static void WriteUInt32ToByteArrayLE(TArray<uint8>& InByteArray, int32& Index, const uint32 Value)
+	{
+		InByteArray[Index++] = (uint8)(Value >> 0);
+		InByteArray[Index++] = (uint8)(Value >> 8);
+		InByteArray[Index++] = (uint8)(Value >> 16);
+		InByteArray[Index++] = (uint8)(Value >> 24);
+	}
+
+	static void WriteUInt16ToByteArrayLE(TArray<uint8>& InByteArray, int32& Index, const uint16 Value)
+	{
+		InByteArray[Index++] = (uint8)(Value >> 0);
+		InByteArray[Index++] = (uint8)(Value >> 8);
+	}
+
+	static void SerializeWaveFile(TArray<uint8>& OutWaveFileData, const uint8* InPCMData, const int32 NumBytes, const int32 NumChannels, const int32 SampleRate, const int32 BytesPerSample, bool bFloatingPoint)
+	{
+		// Reserve space for the raw wave data
+		OutWaveFileData.Empty(NumBytes + 44);
+		OutWaveFileData.AddZeroed(NumBytes + 44);
+
+		int32 WaveDataByteIndex = 0;
+
+		// Wave Format Serialization ----------
+
+		// FieldName: ChunkID
+		// FieldSize: 4 bytes
+		// FieldValue: RIFF (FourCC value, big-endian)
+		OutWaveFileData[WaveDataByteIndex++] = 'R';
+		OutWaveFileData[WaveDataByteIndex++] = 'I';
+		OutWaveFileData[WaveDataByteIndex++] = 'F';
+		OutWaveFileData[WaveDataByteIndex++] = 'F';
+
+		// ChunkName: ChunkSize: 4 bytes
+		// Value: NumBytes + 36. Size of the rest of the chunk following this number. Size of entire file minus 8 bytes.
+		WriteUInt32ToByteArrayLE(OutWaveFileData, WaveDataByteIndex, NumBytes + 36);
+
+		// FieldName: Format
+		// FieldSize: 4 bytes
+		// FieldValue: "WAVE"  (big-endian)
+		OutWaveFileData[WaveDataByteIndex++] = 'W';
+		OutWaveFileData[WaveDataByteIndex++] = 'A';
+		OutWaveFileData[WaveDataByteIndex++] = 'V';
+		OutWaveFileData[WaveDataByteIndex++] = 'E';
+
+		// FieldName: Subchunk1ID
+		// FieldSize: 4 bytes
+		// FieldValue: "fmt "
+		OutWaveFileData[WaveDataByteIndex++] = 'f';
+		OutWaveFileData[WaveDataByteIndex++] = 'm';
+		OutWaveFileData[WaveDataByteIndex++] = 't';
+		OutWaveFileData[WaveDataByteIndex++] = ' ';
+
+		// FieldName: Subchunk1Size
+		// FieldSize: 4 bytes
+		// FieldValue: 16 for PCM
+		WriteUInt32ToByteArrayLE(OutWaveFileData, WaveDataByteIndex, 16);
+
+		// FieldName: AudioFormat
+		// FieldSize: 2 bytes
+		// FieldValue: 1 for PCM, 3 for floating point
+		WriteUInt16ToByteArrayLE(OutWaveFileData, WaveDataByteIndex, bFloatingPoint ? 3 : 1);
+
+		// FieldName: NumChannels
+		// FieldSize: 2 bytes
+		// FieldValue: 1 for for mono
+		WriteUInt16ToByteArrayLE(OutWaveFileData, WaveDataByteIndex, NumChannels);
+
+		// FieldName: SampleRate
+		// FieldSize: 4 bytes
+		// FieldValue: Passed in sample rate
+		WriteUInt32ToByteArrayLE(OutWaveFileData, WaveDataByteIndex, SampleRate);
+
+		// FieldName: ByteRate
+		// FieldSize: 4 bytes
+		// FieldValue: SampleRate * NumChannels * BitsPerSample/8
+		int32 ByteRate = SampleRate * NumChannels * BytesPerSample;
+		WriteUInt32ToByteArrayLE(OutWaveFileData, WaveDataByteIndex, ByteRate);
+
+		// FieldName: BlockAlign
+		// FieldSize: 2 bytes
+		// FieldValue: NumChannels * BitsPerSample/8
+		int32 BlockAlign = NumChannels * BytesPerSample;
+		WriteUInt16ToByteArrayLE(OutWaveFileData, WaveDataByteIndex, BlockAlign);
+
+		// FieldName: BitsPerSample
+		// FieldSize: 2 bytes
+		// FieldValue: BitsPerSample (16 bits per sample)
+		WriteUInt16ToByteArrayLE(OutWaveFileData, WaveDataByteIndex, BytesPerSample * 8);
+
+		// FieldName: Subchunk2ID
+		// FieldSize: 4 bytes
+		// FieldValue: "data" (big endian)
+
+		OutWaveFileData[WaveDataByteIndex++] = 'd';
+		OutWaveFileData[WaveDataByteIndex++] = 'a';
+		OutWaveFileData[WaveDataByteIndex++] = 't';
+		OutWaveFileData[WaveDataByteIndex++] = 'a';
+
+		// FieldName: Subchunk2Size
+		// FieldSize: 4 bytes
+		// FieldValue: number of bytes of the data
+		WriteUInt32ToByteArrayLE(OutWaveFileData, WaveDataByteIndex, NumBytes);
+
+		// Copy the raw PCM data to the audio file
+		FMemory::Memcpy(&OutWaveFileData[WaveDataByteIndex], InPCMData, NumBytes);
+	}
+}
+
+namespace UE::MediaIOAudioDebug
+{
+	static int32 SecondsToDumpFor = 1;
+	static bool bProcessAudio = false;
+
+	static FAutoConsoleCommand CAudioDebug(
+		TEXT("MediaIO.DumpAudio"),
+		TEXT("[seconds (default: 1)] - Number of seconds to dump audio for."),
+		FConsoleCommandWithArgsDelegate::CreateLambda([](const TArray<FString>& Args)
+			{
+				if (Args.Num())
+				{
+					SecondsToDumpFor = FCString::Atoi(*Args[0]);
+				}
+
+				UE_LOG(LogMediaIOAudioDebug, Display, TEXT("Starting to dump audio for %d seconds"), SecondsToDumpFor);
+
+				bProcessAudio = true;
+			}),
+		ECVF_Default);
+
+	struct FAudioDebugDump
+	{
+		template <typename Type>
+		bool Start(const FString& InFileName, int32 InFramesToWrite)
+		{
+			if (bStarted)
+			{
+				return true;
+			}
+
+			const FString Datetime = FDateTime::Now().ToString();
+			SampleSize = sizeof(Type);
+			NumFramesToWrite = InFramesToWrite;
+			BaseFileName = InFileName;
+			FolderPath = FPaths::ProjectSavedDir() / TEXT("AudioDumps") / Datetime;
+
+			bStarted = true;
+			return true;
+		}
+
+
+		template <typename Type>
+		void WriteToWavFile()
+		{
+			if (bStarted && CollectedSamples.Num())
+			{
+				TArray<uint8> WaveFileData;
+				Private::SerializeWaveFile(WaveFileData, (uint8*)CollectedSamples.GetData(), CollectedSamples.Num(), 1, 48000, sizeof(Type), TIsFloatingPoint<Type>::Value);
+
+				const FString WavFileName = BaseFileName + TEXT(".wav");
+
+				if (FArchive* WavWriter = IFileManager::Get().CreateFileWriter(*(FolderPath / WavFileName)))
+				{
+					WavWriter->Serialize((void*)WaveFileData.GetData(), WaveFileData.Num());
+					WavWriter->Close();
+					delete WavWriter;
+				}
+				else
+				{
+					UE_LOG(LogMediaIOAudioDebug, Error, TEXT("Could not write wav file!"));
+				}
+			}
+		}
+
+		template <typename Type>
+		void WriteToCsvFile()
+		{
+			if (bStarted && CollectedSamples.Num())
+			{
+				const FString CsvFileName = BaseFileName + TEXT(".csv");
+
+				if (FArchive* CsvWriter = IFileManager::Get().CreateFileWriter(*(FolderPath / CsvFileName)))
+				{
+					Type* SamplePtr = (Type*)CollectedSamples.GetData();
+
+					for (int32 Index = 0; Index < NumCollectedSamples; Index++)
+					{
+						const FString ValueAsString = FString::Printf(TFormatSpecifier<Type>::GetFormatSpecifier(), *SamplePtr);
+						CsvWriter->Logf(TEXT("%d, %s"), Index, *ValueAsString);
+						SamplePtr++;
+					}
+
+					CsvWriter->Close();
+					delete CsvWriter;
+				}
+				else
+				{
+					UE_LOG(LogMediaIOAudioDebug, Error, TEXT("Could not write wav file!"));
+				}
+			}
+		}
+
+		template <typename Type>
+		void Stop()
+		{
+			WriteToWavFile<Type>();
+			WriteToCsvFile<Type>();
+
+			FrameWrittenIndex = 0;
+			LastWroteIndex = 0;
+
+			bStarted = false;
+		}
+
+		template <typename Type>
+		void CollectSamples(const TArray<Type>& AudioSamples)
+		{
+			if (bStarted)
+			{
+				check(SampleSize == sizeof(Type));
+
+				const int32 NewSamplesSize = AudioSamples.Num() * SampleSize;
+				const int32 OldSize = CollectedSamples.Num();
+				CollectedSamples.AddZeroed(NewSamplesSize);
+
+				uint8* NewSamplesPosition = CollectedSamples.GetData() + OldSize * sizeof(uint8);
+				memcpy(NewSamplesPosition, AudioSamples.GetData(), NewSamplesSize);
+
+				NumCollectedSamples += AudioSamples.Num();
+
+				FrameWrittenIndex++;
+				if (FrameWrittenIndex == NumFramesToWrite)
+				{
+					Stop<Type>();
+				}
+			}
+		}
+
+		/** Collects a single channel. */
+		template <typename Type>
+		void CollectSamples(const uint8* BufferToWrite, int32 BufferSize, uint8 NumChannels)
+		{
+			if (bStarted)
+			{
+				const int32 NewSamplesSize = BufferSize * sizeof(Type);
+				const int32 OldSize = CollectedSamples.Num();
+				CollectedSamples.AddZeroed(NewSamplesSize);
+
+				uint8* NewSamplesPosition = CollectedSamples.GetData() + OldSize * sizeof(uint8);
+				
+				for (int32 Index = 0; Index < BufferSize; Index += NumChannels * sizeof(Type))
+				{
+					for (int32 SampleByteIndex = 0; SampleByteIndex < sizeof(Type); SampleByteIndex++)
+					{
+						CollectedSamples.Add(BufferToWrite[Index + SampleByteIndex]);
+					}
+				}
+				
+				NumCollectedSamples += BufferSize / (sizeof(Type) * NumChannels);
+
+				FrameWrittenIndex++;
+				if (FrameWrittenIndex == NumFramesToWrite)
+				{
+					Stop<Type>();
+				}
+			}
+		}
+		
+
+		FString FolderPath;
+		FString BaseFileName;
+		int32 FrameWrittenIndex = 0;
+		int32 LastWroteIndex = 0;
+		int32 NumFramesToWrite = 0;
+		bool bStarted = false;
+	
+		uint8 SampleSize = 0;
+		TArray<uint8> CollectedSamples;
+		int32 NumCollectedSamples = 0;
+	};
+
+	struct FAudioDebugDumpSystem
+	{
+		TMap<FString, FAudioDebugDump> DebugInfoMap;
+
+		template <typename Type>
+		void StartDebugDump(const FString& BaseFileName)
+		{
+			if (DebugInfoMap.Contains(BaseFileName))
+			{
+				//  Already started
+				return;
+			}
+
+			FAudioDebugDump& DebugDump = DebugInfoMap.FindOrAdd(BaseFileName);
+			if (!DebugDump.bStarted)
+			{
+				int32 Framerate = 24;
+				if (GEngine)
+				{
+					if (GEngine->bUseFixedFrameRate && GEngine->FixedFrameRate != 0)
+					{
+						Framerate = FMath::RoundToInt32(GEngine->FixedFrameRate);
+					}
+				}
+
+				if (Framerate != 0)
+				{
+					DebugDump.Start<Type>(BaseFileName, SecondsToDumpFor * Framerate);
+				}
+			}
+		}
+
+		template <typename Type>
+		void ProcessSamples(const FString& BaseFileName, const TArray<Type>& BufferToWrite)
+		{
+			if (FAudioDebugDump* DebugDump = DebugInfoMap.Find(BaseFileName))
+			{
+				DebugDump->CollectSamples(BufferToWrite);
+
+				if (!DebugDump->bStarted)
+				{
+					UE_LOG(LogMediaIOAudioDebug, Display, TEXT("Finished audio dump for %s"), *DebugDump->BaseFileName);
+					DebugInfoMap.Remove(BaseFileName);
+				}
+			}
+
+			if (DebugInfoMap.Num() == 0)
+			{
+				bProcessAudio = false;
+			}
+		};
+
+		template <typename Type>
+		void ProcessSamples(const FString& BaseFileName, const uint8* BufferToWrite, int32 BufferSize, uint8 NumChannels)
+		{
+			if (FAudioDebugDump* DebugDump = DebugInfoMap.Find(BaseFileName))
+			{
+				DebugDump->CollectSamples<Type>(BufferToWrite, BufferSize, NumChannels);
+
+				if (!DebugDump->bStarted)
+				{
+					UE_LOG(LogMediaIOAudioDebug, Display, TEXT("Finished audio dump for %s"), *DebugDump->BaseFileName);
+					DebugInfoMap.Remove(BaseFileName);
+				}
+			}
+
+			if (DebugInfoMap.Num() == 0)
+			{
+				bProcessAudio = false;
+			}
+		};
+
+		template <typename Type>
+		void ProcessAudio(const FString& BaseFileName, const TArray<Type>& BufferToWrite)
+		{
+			if (bProcessAudio)
+			{
+				StartDebugDump<Type>(BaseFileName);
+				ProcessSamples(BaseFileName, BufferToWrite);
+			}
+		}
+		
+		template <typename Type>
+        void ProcessAudio(const FString& BaseFileName, const uint8* BufferToWrite, int32 BufferSize, uint8 NumChannels)
+        {
+        	if (bProcessAudio)
+        	{
+        		StartDebugDump<Type>(BaseFileName);
+        		ProcessSamples<Type>(BaseFileName, BufferToWrite, BufferSize, NumChannels);
+        	}
+        }
+	};
+
+	static FAudioDebugDumpSystem& GetSingleton()
+	{
+		static FAudioDebugDumpSystem DebugSystem;
+		return DebugSystem;
+	}
+}
diff --git a/Engine/Plugins/Media/MediaIOFramework/Source/MediaIOCore/Public/MediaIOCoreAudioOutput.h b/Engine/Plugins/Media/MediaIOFramework/Source/MediaIOCore/Public/MediaIOCoreAudioOutput.h
--- a/Engine/Plugins/Media/MediaIOFramework/Source/MediaIOCore/Public/MediaIOCoreAudioOutput.h
+++ b/Engine/Plugins/Media/MediaIOFramework/Source/MediaIOCore/Public/MediaIOCoreAudioOutput.h
@@ -6,68 +6,78 @@
 #include "ISubmixBufferListener.h"
 #include "Math/NumericLimits.h"
 
-#if UE_ENABLE_INCLUDE_ORDER_DEPRECATED_IN_5_1
-#include "AudioDevice.h"
+#if WITH_MEDIA_IO_AUDIO_DEBUGGING
+#include "MediaIOAudioDebug.h"
 #endif
 
-namespace UE::MediaIoCoreModule::Private
-{
+class FAudioDevice;
+
+class MEDIAIOCORE_API FMediaIOAudioOutput
+{ 
+public:
+	struct FAudioOptions
+	{
+		uint32 InNumInputChannels = 0;
+		uint32 InNumOutputChannels = 0;
+		FFrameRate InTargetFrameRate;
+		uint32 InMaxSampleLatency = 0;
+		uint32 InOutputSampleRate = 0;
+	};
+
+	FMediaIOAudioOutput(Audio::FPatchOutputStrongPtr InPatchOutput, const FAudioOptions& InAudioOptions);
+
 	template <typename OutputType>
-	TArray<OutputType> ConvertAndUpmixBuffer(const Audio::FAlignedFloatBuffer& InBuffer, int32 NumInputChannels, int32 NumOutputChannels)
+	TArray<OutputType> ConvertAndUpmixBuffer(const Audio::FAlignedFloatBuffer& InBuffer, int32 InNumInputChannels, int32 InNumOutputChannels) const
 	{
 		TRACE_CPUPROFILER_EVENT_SCOPE(MediaIoCoreModule::ConvertAndUpmixBuffer);
+
 		/**
 		 * @Note: Audio::DownmixBuffer was not used here because it does not convert data in place,
-		 *  this forces us to create an additional buffer to hold the final int32 values. Instead we 
+		 *  this forces us to create an additional buffer to hold the final int32 values. Instead we
 		 *  upmix and convert the values in a single pass.
 		 */
 
-		// @Note: We only support upmixing at the moment.
-		
-		TArray<OutputType> ConvertedBuffer;
-		const float ChannelRatio = static_cast<float>(NumOutputChannels) / NumInputChannels;
-		ConvertedBuffer.SetNumZeroed(FMath::CeilToInt32(ChannelRatio * InBuffer.Num()));
-
-		OutputType* ConvertedBufferPtr = ConvertedBuffer.GetData();
+		 // @Note: We only support upmixing at the moment.
 
-		for (int32 Index = 0; Index < InBuffer.Num(); Index += NumInputChannels)
+		TArray<OutputType> ConvertedBuffer;
+		if (InBuffer.Num())
 		{
-			constexpr double Scale = TNumericLimits<OutputType>::Max();
+			const float ChannelRatio = static_cast<float>(InNumOutputChannels) / InNumInputChannels;
+			ConvertedBuffer.SetNumZeroed(FMath::CeilToInt32(ChannelRatio * InBuffer.Num()));
+
+			OutputType* ConvertedBufferPtr = ConvertedBuffer.GetData();
 
-			// @Note: While this conversion is a common approach, it can introduce distortion. 
-			// See: https://www.cs.cmu.edu/~rbd/papers/cmj-float-to-int.html
-			// Convert values for each channel
-			for (int32 ChannelIndex = 0; ChannelIndex < NumInputChannels; ++ChannelIndex)
+			for (int32 Index = 0; Index < InBuffer.Num(); Index += InNumInputChannels)
 			{
-				const float FloatValue = InBuffer[Index + ChannelIndex];
-				OutputType ConvertedValue = static_cast<OutputType>((FloatValue * Scale) + 0.5);
+				constexpr double Scale = TNumericLimits<OutputType>::Max();
+
+				// @Note: While this conversion is a common approach, it can introduce distortion. 
+				// See: https://www.cs.cmu.edu/~rbd/papers/cmj-float-to-int.html
+				// Convert values for each channel
+				for (int32 ChannelIndex = 0; ChannelIndex < InNumInputChannels; ++ChannelIndex)
+				{
+					const float FloatValue = InBuffer[Index + ChannelIndex];
+					OutputType ConvertedValue = static_cast<OutputType>((FloatValue * Scale) + 0.5);
+
+					*(ConvertedBufferPtr + ChannelIndex) = ConvertedValue;
+				}
 
-				*(ConvertedBufferPtr + ChannelIndex) = ConvertedValue;
+				ConvertedBufferPtr += InNumOutputChannels;
 			}
-		
-			ConvertedBufferPtr += NumOutputChannels;
+
+			const int32 FinalNumSamples = AlignDown(ConvertedBuffer.Num(), 4);
+			constexpr bool bAllowShrinking = true;
+			ConvertedBuffer.SetNum(FinalNumSamples, bAllowShrinking);
+
+#if WITH_MEDIA_IO_AUDIO_DEBUGGING
+			//MEDIA_IO_DUMP_AUDIO(InBuffer.GetData(), InBuffer.Num() * sizeof(float), sizeof(float), NumInputChannels);
+			UE::MediaIOAudioDebug::GetSingleton().ProcessAudio<float>(TEXT("InitialBuffer"), (uint8*) InBuffer.GetData(), InBuffer.Num() * sizeof(float), InNumInputChannels);
+			UE::MediaIOAudioDebug::GetSingleton().ProcessAudio<OutputType>(TEXT("ConvertedBuffer"), (uint8*) ConvertedBuffer.GetData(), ConvertedBuffer.Num() * sizeof(OutputType), InNumOutputChannels);
+#endif /*WITH_MEDIA_IO_AUDIO_DEBUGGING*/
 		}
 
-		const int32 FinalNumSamples = AlignDown(ConvertedBuffer.Num(), 4);
-		constexpr bool bAllowShrinking = true;
-		ConvertedBuffer.SetNum(FinalNumSamples, bAllowShrinking);
 		return ConvertedBuffer;
 	}
-}
-
-class MEDIAIOCORE_API FMediaIOAudioOutput
-{ 
-public:
-	struct FAudioOptions
-	{
-		uint32 InNumInputChannels = 0;
-		uint32 InNumOutputChannels = 0;
-		FFrameRate InTargetFrameRate;
-		uint32 InMaxSampleLatency = 0;
-		uint32 InOutputSampleRate = 0;
-	};
-
-	FMediaIOAudioOutput(Audio::FPatchOutputStrongPtr InPatchOutput, const FAudioOptions& InAudioOptions);
 
 	/**
 	 * Get the audio sample that were accumulated.
@@ -78,7 +88,7 @@ public:
 		TRACE_CPUPROFILER_EVENT_SCOPE(FMediaIOAudioOutput::GetAudioSamples);
 		
 		const Audio::FAlignedFloatBuffer FloatBuffer = GetFloatBuffer(NumSamplesPerFrame);
-		return UE::MediaIoCoreModule::Private::ConvertAndUpmixBuffer<OutputType>(FloatBuffer, NumInputChannels, NumOutputChannels);
+		return ConvertAndUpmixBuffer<OutputType>(FloatBuffer, NumInputChannels, NumOutputChannels);
 	}
 	
 	template <typename OutputType>
@@ -87,21 +97,23 @@ public:
 		TRACE_CPUPROFILER_EVENT_SCOPE(FMediaIOAudioOutput::GetAudioSamples);
 		
 		const Audio::FAlignedFloatBuffer FloatBuffer = GetFloatBuffer(NumSamplesToGet);
-		return UE::MediaIoCoreModule::Private::ConvertAndUpmixBuffer<OutputType>(FloatBuffer, NumInputChannels, NumOutputChannels);
+		return ConvertAndUpmixBuffer<OutputType>(FloatBuffer, NumInputChannels, NumOutputChannels);
+	}
+
+	template <typename OutputType>
+	TArray<OutputType> GetAllAudioSamples() const
+	{
+		TRACE_CPUPROFILER_EVENT_SCOPE(FMediaIOAudioOutput::GetAllAudioSamples);
+
+		const Audio::FAlignedFloatBuffer FloatBuffer = GetFloatBuffer(MaxSampleLatency);
+		return ConvertAndUpmixBuffer<OutputType>(FloatBuffer, NumInputChannels, NumOutputChannels);
 	}
 
+
 public:
 	// @todo: Depend on frame number to correctly fetch the right amount of frames on framerates like 59.97
 	int32 NumSamplesPerFrame = 0;
 
-private:
-	int32 GetAudioBuffer(int32 InNumSamplesToPop, float* OutBuffer) const;
-	Audio::FAlignedFloatBuffer GetFloatBuffer(uint32 NumSamplesToGet) const;
-
-private:
-	/** The buffer accumulating audio samples. */
-	Audio::FPatchOutputStrongPtr PatchOutput;
-
 	/** Number of audio channels on the engine side. */
 	int32 NumInputChannels;
 
@@ -114,6 +126,15 @@ private:
 	uint32 MaxSampleLatency;
 
 	uint32 OutputSampleRate;
+
+private:
+	int32 GetAudioBuffer(int32 InNumSamplesToPop, float* OutBuffer) const;
+	Audio::FAlignedFloatBuffer GetFloatBuffer(uint32 NumSamplesToGet) const;
+
+private:
+	/** The buffer accumulating audio samples. */
+	Audio::FPatchOutputStrongPtr PatchOutput;
+
 };
 
 /**
