Fix CIS warning by removing Engine dependency on AjaCore

diff --git a/Engine/Plugins/Media/AjaMedia/Source/Aja/AjaCore.Build.cs b/Engine/Plugins/Media/AjaMedia/Source/Aja/AjaCore.Build.cs
--- a/Engine/Plugins/Media/AjaMedia/Source/Aja/AjaCore.Build.cs
+++ b/Engine/Plugins/Media/AjaMedia/Source/Aja/AjaCore.Build.cs
@@ -28,8 +28,7 @@ public class AjaCore : ModuleRules
 		PublicDependencyModuleNames.AddRange(
 			new string[] {
 				"Core",
-				"RHI",
-				"Engine"
+				"RHI"
 			});
 
 		bDisableStaticAnalysis = true;
