Aja audio - Add path for direct audio write

diff --git a/Engine/Plugins/Media/AjaMedia/Source/Aja/Private/AJA.cpp b/Engine/Plugins/Media/AjaMedia/Source/Aja/Private/AJA.cpp
--- a/Engine/Plugins/Media/AjaMedia/Source/Aja/Private/AJA.cpp
+++ b/Engine/Plugins/Media/AjaMedia/Source/Aja/Private/AJA.cpp
@@ -444,6 +444,11 @@ namespace AJA
 		return Channel ? Channel->SetVideoFrameData(InFrameData, RHITexture) : false;
 	}
 
+	bool AJAOutputChannel::DMAWriteAudio(const uint8* Buffer, int32 BufferSize)
+	{
+		return Channel ? Channel->DMAWriteAudio(Buffer, BufferSize) : false;
+	}
+
 	bool AJAOutputChannel::GetOutputDimension(uint32_t& OutWidth, uint32_t& OutHeight) const
 	{
 		return Channel ? Channel->GetOutputDimension(OutWidth, OutHeight) : false;
diff --git a/Engine/Plugins/Media/AjaMedia/Source/Aja/Private/ChannelBase.cpp b/Engine/Plugins/Media/AjaMedia/Source/Aja/Private/ChannelBase.cpp
--- a/Engine/Plugins/Media/AjaMedia/Source/Aja/Private/ChannelBase.cpp
+++ b/Engine/Plugins/Media/AjaMedia/Source/Aja/Private/ChannelBase.cpp
@@ -291,6 +291,12 @@ namespace AJA
 				return false;
 			}
 
+			if (Options.bDirectlyWriteAudio && Options.bUseAutoCirculating)
+			{
+				UE_LOG(LogTemp, Error, TEXT("ConfigureVideo: Direct audio write is only available outside of auto-circulate mode..\n"));
+				return false;
+			}
+			
 			const bool bConnectChannel = true;
 			const bool bAsOwner = true;
 
diff --git a/Engine/Plugins/Media/AjaMedia/Source/Aja/Private/OutputChannel.cpp b/Engine/Plugins/Media/AjaMedia/Source/Aja/Private/OutputChannel.cpp
--- a/Engine/Plugins/Media/AjaMedia/Source/Aja/Private/OutputChannel.cpp
+++ b/Engine/Plugins/Media/AjaMedia/Source/Aja/Private/OutputChannel.cpp
@@ -7,6 +7,10 @@
 #include "SSE2MemCpy.h"
 #include <stdlib.h>
 
+DECLARE_FLOAT_COUNTER_STAT(TEXT("AJA Audio Delay (s)"), STAT_AjaMediaCapture_Audio_Delay, STATGROUP_Media);
+
+#include "Misc/ScopeLock.h"
+
 namespace AJA
 {
 	namespace Private
@@ -89,6 +93,15 @@ namespace AJA
 			return false;
 		}
 
+		bool OutputChannel::DMAWriteAudio(const uint8_t* InAudioBufer, int32_t InAudioBufferSize)
+		{
+			if (ChannelThread)
+			{
+				return ChannelThread->DMAWriteAudio(InAudioBufer, InAudioBufferSize);
+			}
+			return false;
+		}
+
 		bool OutputChannel::SetVideoFrameData(const AJAOutputFrameBufferData& InFrameData, FRHITexture* RHITexture)
 		{
 			if (ChannelThread)
@@ -442,7 +455,8 @@ namespace AJA
 			{
 				return false;
 			}
-			if (UseAudio() && (!InFrame->bAudioLineFilled))
+			// When writing audio directly, no need to check if the audio line is filled since we don't use it.
+			if (UseAudio() && !InFrame->bAudioLineFilled && !GetOptions().bDirectlyWriteAudio)
 			{
 				return false;
 			}
@@ -533,7 +547,7 @@ namespace AJA
 			{
 				return false;
 			}
-
+			
 			if (InAudioBufferSize > AudioBufferSize)
 			{
 				UE_LOG(LogTemp, Error, TEXT("SetAudioFrameData: Can't set the audio. The buffer size is not the same as the AJA Audio Format for output channel %d on device %S.\n"), uint32_t(Channel) + 1, GetDevice().GetDisplayName().c_str());
@@ -550,12 +564,11 @@ namespace AJA
 			AJA_CHECK(AvailableWritingFrame->Audio[AudioBufferSize] == AJA_TEST_MEMORY_END_TAG);
 #endif
 
-			if (InAudioBuffer && InAudioBufferSize != 0)
+			if (InAudioBuffer && InAudioBufferSize != 0 && ensure(AvailableWritingFrame->AudioBuffer))
 			{
 				memcpy(AvailableWritingFrame->AudioBuffer, InAudioBuffer, InAudioBufferSize);
 			}
 
-
 			AvailableWritingFrame->bAudioLineFilled = true;
 			AvailableWritingFrame->CopiedAudioBufferSize = InAudioBufferSize;
 			PushWhenFrameReady(AvailableWritingFrame);
@@ -751,17 +764,20 @@ namespace AJA
 
 				if (ensureMsgf(AudioBufferSize > 0, TEXT("Audio could not be initialized.")))
 				{
-					AJAAutoLock AutoLock(&FrameLock);
-					for (Frame* IttFrame : AllFrames)
+					if (!GetOptions().bDirectlyWriteAudio)
 					{
-						AJA_CHECK(IttFrame->AudioBuffer == nullptr);
-	#if AJA_TEST_MEMORY_BUFFER
-						IttFrame->AudioBuffer = (uint8_t*)AJAMemory::AllocateAligned(AudioBufferSize + 1, AJA_PAGE_SIZE);
-						IttFrame->AudioBuffer[AudioBufferSize] = AJA_TEST_MEMORY_END_TAG;
-	#else
-						IttFrame->AudioBuffer = (uint8_t*)AJAMemory::AllocateAligned(AudioBufferSize, AJA_PAGE_SIZE);
-	#endif
-						memset(IttFrame->AudioBuffer, 0, AudioBufferSize);
+						AJAAutoLock AutoLock(&FrameLock);
+						for (Frame* IttFrame : AllFrames)
+						{
+							AJA_CHECK(IttFrame->AudioBuffer == nullptr);
+		#if AJA_TEST_MEMORY_BUFFER
+							IttFrame->AudioBuffer = (uint8_t*)AJAMemory::AllocateAligned(AudioBufferSize + 1, AJA_PAGE_SIZE);
+							IttFrame->AudioBuffer[AudioBufferSize] = AJA_TEST_MEMORY_END_TAG;
+		#else
+							IttFrame->AudioBuffer = (uint8_t*)AJAMemory::AllocateAligned(AudioBufferSize, AJA_PAGE_SIZE);
+		#endif
+							memset(IttFrame->AudioBuffer, 0, AudioBufferSize);
+						}
 					}
 
 					if (GetOptions().bUseAutoCirculating)
@@ -1106,6 +1122,8 @@ namespace AJA
 					bIsValidWaitForOutput = (bRunning && FieldId == NTV2_FIELD0);
 				}
 
+				FScopeLock DeviceLock(&DeviceCriticalSection);
+
 				if (!bRunning)
 				{
 					UE_LOG(LogTemp, Error, TEXT("PingPong: Can't wait for the output field for channel %d on device %S.\n"), uint32_t(Channel) + 1, GetDevice().GetDisplayName().c_str());
@@ -1173,7 +1191,7 @@ namespace AJA
 						}
 					}
 
-					if (UseAudio() && AvailableReadingFrame->CopiedAudioBufferSize)
+					if (UseAudio() && AvailableReadingFrame->CopiedAudioBufferSize && !GetOptions().bDirectlyWriteAudio)
 					{
 						bRunning &= Thread_HandleAudio(TEXT("PingPong"), AvailableReadingFrame);
 					}
@@ -1220,7 +1238,7 @@ namespace AJA
 				}
 				else if (GetOptions().bDisplayWarningIfDropFrames)
 				{
-					if (UseAudio())
+					if (UseAudio() && !GetOptions().bDirectlyWriteAudio)
 					{
 						Thread_HandleLostFrameAudio();
 					}
@@ -1341,20 +1359,18 @@ namespace AJA
 		}
 
 
-		bool OutputChannelThread::Thread_TransferAudioBuffer(uint8* InBuffer, int32 InBufferSize)
+		bool OutputChannelThread::Thread_TransferAudioBuffer(const uint8* InBuffer, int32 InBufferSize)
 		{
 			bool bSuccess = true;
 
 			ULWord AudioWrapOffset = 0;
-			ULWord WriteHeadPosition = 0;
-			ULWord PlayHeadPosition = 0;
 
 			bSuccess &= GetDevice().GetAudioWrapAddress(AudioWrapOffset);
 
 			if (CurrentAudioWriteOffset + InBufferSize < AudioWrapOffset)
 			{
 				// Simplest case, just write the whole data to the SDRAM.
-				bSuccess &= GetDevice().DMAWriteAudio(AudioSystem, reinterpret_cast<ULWord*>(InBuffer), CurrentAudioWriteOffset, InBufferSize);
+				bSuccess &= GetDevice().DMAWriteAudio(AudioSystem, reinterpret_cast<const ULWord*>(InBuffer), CurrentAudioWriteOffset, InBufferSize);
 				CurrentAudioWriteOffset += InBufferSize;
 			}
 			else
@@ -1368,10 +1384,10 @@ namespace AJA
 					RemainingBytes = 0;
 				}
 
-				bSuccess &= GetDevice().DMAWriteAudio(AudioSystem, reinterpret_cast<ULWord*>(InBuffer), CurrentAudioWriteOffset, BytesUntilBufferEnd);
+				bSuccess &= GetDevice().DMAWriteAudio(AudioSystem, reinterpret_cast<const ULWord*>(InBuffer), CurrentAudioWriteOffset, BytesUntilBufferEnd);
 				if (RemainingBytes > 0)
 				{
-					bSuccess &= GetDevice().DMAWriteAudio(AudioSystem, reinterpret_cast<ULWord*>(InBuffer) + BytesUntilBufferEnd, 0, RemainingBytes);
+					bSuccess &= GetDevice().DMAWriteAudio(AudioSystem, reinterpret_cast<const ULWord*>(InBuffer) + BytesUntilBufferEnd, 0, RemainingBytes);
 				}
 
 				CurrentAudioWriteOffset = RemainingBytes;
@@ -1401,7 +1417,7 @@ namespace AJA
 			}
 			else if (!bAudioOutputRunning)
 			{
-				// Only offset the start ptr if this is the first time we're starting the audio output.
+				// Only offset the start ptr if this is1 the first time we're starting the audio output.
 
 				ULWord WrapOffset = 0;
 				bAudioSuccess &= GetDevice().GetAudioWrapAddress(WrapOffset);
@@ -1419,8 +1435,12 @@ namespace AJA
 			{
 				double OffsetInSeconds = 0;
 				bAudioSuccess &= Thread_GetAudioOffsetInSeconds(OffsetInSeconds);
-				bAudioSuccess &= GetDevice().ReadAudioLastOut(AudioPlayheadLastPosition, AudioSystem);
-				UE_LOG(LogTemp, Verbose, TEXT("%s: Play head: %d, Write Head: %d, Offset: %f"), *OutputMethod, AudioPlayheadLastPosition, CurrentAudioWriteOffset, OffsetInSeconds);
+				ULWord PlayHeadPosition = 0;
+				bAudioSuccess &= GetDevice().ReadAudioLastOut(PlayHeadPosition, AudioSystem);
+				AudioPlayheadLastPosition = PlayHeadPosition;
+				
+				SET_FLOAT_STAT(STAT_AjaMediaCapture_Audio_Delay, OffsetInSeconds);
+				UE_LOG(LogTemp, Verbose, TEXT("%s: Play head: %d, Write Head: %d, Offset: %f"), *OutputMethod, AudioPlayheadLastPosition.load(), CurrentAudioWriteOffset.load(), OffsetInSeconds);
 
 				bAudioTransferWasSuccessful &= Thread_TransferAudioBuffer(AvailableReadingFrame->AudioBuffer, AvailableReadingFrame->CopiedAudioBufferSize);
 
@@ -1434,7 +1454,7 @@ namespace AJA
 					bool bPausePlayout = true;
 					GetDevice().SetAudioOutputPause(AudioSystem, bPausePlayout);
 
-					UE_LOG(LogTemp, Verbose, TEXT("%s Pausing audio output because the play head is catching up to the write head. Last playhead position: %d"), *OutputMethod, AudioPlayheadLastPosition);
+					UE_LOG(LogTemp, Verbose, TEXT("%s Pausing audio output because the play head is catching up to the write head. Last playhead position: %d"), *OutputMethod, AudioPlayheadLastPosition.load());
 				}
 			}
 			else
@@ -1445,6 +1465,61 @@ namespace AJA
 			return bAudioTransferWasSuccessful;
 		}
 
+		bool OutputChannelThread::DMAWriteAudio(const uint8_t* InAudioBuffer, int32_t BufferSize)
+		{
+			if (!Device || !Device->GetCard())
+			{
+				return false;
+			}
+
+			TRACE_CPUPROFILER_EVENT_SCOPE(OutputChannelThread::DMAWriteAudio);
+			FScopeLock DeviceLock(&DeviceCriticalSection);
+
+			bool bAudioSuccess = true;
+
+			bool bAudioOutputRunning = false;
+			bAudioSuccess &= GetDevice().IsAudioOutputRunning(AudioSystem, bAudioOutputRunning);
+
+			UE_CLOG(!bAudioSuccess, LogTemp, Verbose, TEXT("IsAudioOutputRunning returned false."));
+
+			bool bIsPaused = false;
+			bAudioSuccess &= GetDevice().GetAudioOutputPause(AudioSystem, bIsPaused);
+
+			UE_CLOG(!bAudioSuccess, LogTemp, Verbose, TEXT("GetAudioOutputPause returned false."));
+
+			// Start audio output as late as possible to reduce delay.
+			if (!bAudioOutputRunning)
+			{
+				// Only offset the start ptr if this is the first time we're starting the audio output.
+				ULWord WrapOffset = 0;
+				bAudioSuccess &= GetDevice().GetAudioWrapAddress(WrapOffset);
+
+				constexpr int32 NumOffsetFrames = 4;
+				const int32 InitialOffset = Options.NumberOfAudioChannel * NumSamplesPerFrame * NumOffsetFrames * sizeof(int32);
+
+				CurrentAudioWriteOffset = InitialOffset;
+
+				UE_LOG(LogTemp, Verbose, TEXT("PingPong: Starting audio output."));
+				bAudioSuccess &= GetDevice().StartAudioOutput(AudioSystem); // This resets the playhead to 0, we have to override that to something closer to the write head in order to reduce audio delay
+			}
+
+			double OffsetInSeconds = 0;
+			bAudioSuccess &= Thread_GetAudioOffsetInSeconds(OffsetInSeconds);
+			UE_CLOG(!bAudioSuccess, LogTemp, Verbose, TEXT("Thread_GetAudioOffsetInSeconds returned false."));
+			ULWord PlayHeadPosition = 0;
+			bAudioSuccess &= GetDevice().ReadAudioLastOut(PlayHeadPosition, AudioSystem);
+			UE_CLOG(!bAudioSuccess, LogTemp, Verbose, TEXT("ReadAudioLastOut returned false."));
+			
+			AudioPlayheadLastPosition = PlayHeadPosition;
+			SET_FLOAT_STAT(STAT_AjaMediaCapture_Audio_Delay, OffsetInSeconds);
+			UE_LOG(LogTemp, Verbose, TEXT("%s: Play head: %d, Write Head: %d, Offset: %f"), TEXT("PingPong"), AudioPlayheadLastPosition.load(), CurrentAudioWriteOffset.load(), OffsetInSeconds);
+
+			bAudioSuccess &= Thread_TransferAudioBuffer(InAudioBuffer, BufferSize);
+			
+			return bAudioSuccess;
+		}
+		
+
 		bool OutputChannelThread::Thread_HandleLostFrameAudio()
 		{
 			bool bAudioSuccess = true;
@@ -1461,7 +1536,7 @@ namespace AJA
 				{
 					constexpr bool bPausePlayout = true;
 					bAudioSuccess &= GetDevice().SetAudioOutputPause(AudioSystem, bPausePlayout);
-					UE_LOG(LogTemp, Warning, TEXT("Pausing audio output. Last playhead position: %d"), AudioPlayheadLastPosition);
+					UE_LOG(LogTemp, Warning, TEXT("Pausing audio output. Last playhead position: %d"), AudioPlayheadLastPosition.load());
 				}
 			}
 
@@ -1470,7 +1545,7 @@ namespace AJA
 
 		bool OutputChannelThread::Thread_ShouldPauseAudioOutput()
 		{
-			int32 AudioBufferZoneInFrames = 1;
+			constexpr int32 AudioBufferZoneInFrames = 1;
 			int32 OffsetInSamples = 0;
 			Thread_GetAudioOffsetInSamples(OffsetInSamples);
 
diff --git a/Engine/Plugins/Media/AjaMedia/Source/Aja/Private/OutputChannel.h b/Engine/Plugins/Media/AjaMedia/Source/Aja/Private/OutputChannel.h
--- a/Engine/Plugins/Media/AjaMedia/Source/Aja/Private/OutputChannel.h
+++ b/Engine/Plugins/Media/AjaMedia/Source/Aja/Private/OutputChannel.h
@@ -31,6 +31,7 @@ namespace AJA
 			bool SetAncillaryFrameData(const AJAOutputFrameBufferData& InFrameData, uint8_t* AncillaryBuffer, uint32_t AncillaryBufferSize);
 			bool SetAudioFrameData(const AJAOutputFrameBufferData& InFrameData, uint8_t* InAudioBuffer, uint32_t InAudioBufferSize);
 			bool SetVideoFrameData(const AJAOutputFrameBufferData& InFrameData, uint8_t* InVideoBuffer, uint32_t InVideoBufferSize);
+			bool DMAWriteAudio(const uint8_t* InAudioBuffer, int32_t InAudioBufferSize);
 			bool SetVideoFrameData(const AJAOutputFrameBufferData& InFrameData, FRHITexture* RHITexture);
 			bool GetOutputDimension(uint32_t& OutWidth, uint32_t& OutHeight) const;
 			int32_t GetNumAudioSamplesPerFrame(const AJAOutputFrameBufferData& InFrameData) const;
@@ -58,6 +59,7 @@ namespace AJA
 			bool SetAudioFrameData(const AJAOutputFrameBufferData& InFrameData, uint8_t* AudioBuffer, uint32_t AudioBufferSize);
 			bool SetVideoFrameData(const AJAOutputFrameBufferData& InFrameData, uint8_t* VideoBuffer, uint32_t VideoBufferSize);
 			bool SetVideoFrameData(const AJAOutputFrameBufferData& InFrameData, FRHITexture* RHITexture);
+			bool DMAWriteAudio(const uint8_t* InAudioBuffer, int32_t BufferSize);
 
 		protected:
 			virtual bool DeviceThread_ConfigureAnc(DeviceConnection::CommandList& InCommandList) override;
@@ -84,7 +86,7 @@ namespace AJA
 			bool Thread_GetAudioOffset(int32& OutAudioOffset);
 			bool Thread_GetAudioOffsetInSeconds(double& OutAudioOffset);
 			bool Thread_GetAudioOffsetInSamples(int32& OutAudioOffset);
-			bool Thread_TransferAudioBuffer(uint8* InBuffer, int32 InBufferSize);
+			bool Thread_TransferAudioBuffer(const uint8* InBuffer, int32 InBufferSize);
 			bool Thread_HandleAudio(const FString& OutputMethod, Frame* AvailableReadingFrame);
 			bool Thread_HandleLostFrameAudio();
 			bool Thread_ShouldPauseAudioOutput();
@@ -135,9 +137,12 @@ namespace AJA
 
 			bool bDMABuffersRegistered = false;
 
-			ULWord CurrentAudioWriteOffset = 0;
+			std::atomic<ULWord> CurrentAudioWriteOffset = 0;
 			int32_t NumSamplesPerFrame = 0;
-			ULWord AudioPlayheadLastPosition = 0;
+			std::atomic<ULWord> AudioPlayheadLastPosition = 0;
+
+			/** Used to protect access to the device when accessed by the audio thread. */
+			FCriticalSection DeviceCriticalSection;
 
 			UE::GPUTextureTransfer::TextureTransferPtr TextureTransfer = nullptr;
 		};
diff --git a/Engine/Plugins/Media/AjaMedia/Source/Aja/Public/AJALib.h b/Engine/Plugins/Media/AjaMedia/Source/Aja/Public/AJALib.h
--- a/Engine/Plugins/Media/AjaMedia/Source/Aja/Public/AJALib.h
+++ b/Engine/Plugins/Media/AjaMedia/Source/Aja/Public/AJALib.h
@@ -420,6 +420,7 @@ namespace AJA
 				uint32_t bTEST_OutputInterlaced : 1; // when outputting, warn if the field 1 and field 2 are the same color. It except one of the line to be white and the other line to be not white
 				uint32_t bUseGPUDMA : 1; // Whether to use GPU Direct when outputting.
 				uint32_t bAutoDetectFormat : 1; // Whether to autodetect format on input.
+				uint32_t bDirectlyWriteAudio : 1; // Experimental mode to set audio from the audio thread directly.
 			};
 			uint32_t Options;
 		};
@@ -482,6 +483,7 @@ namespace AJA
 		bool SetAudioFrameData(const AJAOutputFrameBufferData& InFrameData, uint8_t* AudioBuffer, uint32_t AudioBufferSize);
 		bool SetVideoFrameData(const AJAOutputFrameBufferData& InFrameData, uint8_t* VideoBuffer, uint32_t VideoBufferSize);
 		bool SetVideoFrameData(const AJAOutputFrameBufferData& InFrameData, FRHITexture* RHITexture);
+		bool DMAWriteAudio(const uint8* Buffer, int32 BufferSize);
 
 		bool GetOutputDimension(uint32_t& OutWidth, uint32_t& OutHeight) const;
 		int32_t GetNumAudioSamplesPerFrame(const AJAOutputFrameBufferData& InFrameData) const;
diff --git a/Engine/Plugins/Media/AjaMedia/Source/AjaMediaOutput/Private/AjaMediaCapture.cpp b/Engine/Plugins/Media/AjaMedia/Source/AjaMediaOutput/Private/AjaMediaCapture.cpp
--- a/Engine/Plugins/Media/AjaMedia/Source/AjaMediaOutput/Private/AjaMediaCapture.cpp
+++ b/Engine/Plugins/Media/AjaMedia/Source/AjaMediaOutput/Private/AjaMediaCapture.cpp
@@ -268,7 +268,14 @@ bool UAjaMediaCapture::ValidateMediaOutput() const
 bool UAjaMediaCapture::InitializeCapture()
 {
 	UAjaMediaOutput* AjaMediaSource = CastChecked<UAjaMediaOutput>(MediaOutput);
+	
 	const bool bResult = InitAJA(AjaMediaSource);
+
+	if (GEngine)
+	{
+		GEngine->GetEngineSubsystem<UMediaIOCoreSubsystem>()->OnBufferReceived_AudioThread().AddUObject(this, &UAjaMediaCapture::OnAudioBufferReceived_AudioThread);
+	}
+	
 	if (bResult)
 	{
 #if WITH_EDITOR
@@ -341,6 +348,11 @@ void UAjaMediaCapture::StopCaptureImpl(bool bAllowPendingFrameToBeProcess)
 			}
 		}
 
+		if (GEngine)
+		{
+			GEngine->GetEngineSubsystem<UMediaIOCoreSubsystem>()->OnBufferReceived_AudioThread().RemoveAll(this);
+		}
+
 		AudioOutput.Reset();
 
 		RestoreViewportTextureAlpha(GetCapturingSceneViewport());
@@ -470,8 +482,10 @@ bool UAjaMediaCapture::InitAJA(UAjaMediaOutput* InAjaMediaOutput)
 	ChannelOptions.TimecodeFormat =  AjaMediaCaptureUtils::ConvertTimecode(InAjaMediaOutput->TimecodeFormat);
 	ChannelOptions.OutputReferenceType = AjaMediaCaptureUtils::Convert(InAjaMediaOutput->OutputConfiguration.OutputReference);
 	ChannelOptions.bUseGPUDMA = ShouldCaptureRHIResource();
+	ChannelOptions.bDirectlyWriteAudio = InAjaMediaOutput->bOutputAudioOnAudioThread;
 
 	bOutputAudio = InAjaMediaOutput->bOutputAudio;
+	bDirectlyWriteAudio = InAjaMediaOutput->bOutputAudioOnAudioThread;
 	
 	if (bOutputAudio)
 	{
@@ -481,6 +495,12 @@ bool UAjaMediaCapture::InitAJA(UAjaMediaOutput* InAjaMediaOutput)
 		}
 	}
 	
+	if (GEngine)
+	{
+		NumInputChannels = GEngine->GetEngineSubsystem<UMediaIOCoreSubsystem>()->GetNumAudioInputChannels();
+		NumOutputChannels = static_cast<int32>(InAjaMediaOutput->NumOutputAudioChannels);
+	}
+	
 	PixelFormat = InAjaMediaOutput->PixelFormat;
 	UseKey = ChannelOptions.bUseKey;
 	
@@ -581,15 +601,19 @@ void UAjaMediaCapture::LockDMATexture_RenderThread(FTextureRHIRef InTexture)
 			
 			TextureTransfer->RegisterTexture(Args);
 		}
+
+		TextureTransfer->LockTexture(InTexture->GetTexture2D());
 	}
 
-	TextureTransfer->LockTexture(InTexture->GetTexture2D());
 }
 
 void UAjaMediaCapture::UnlockDMATexture_RenderThread(FTextureRHIRef InTexture)
 {
-	TRACE_CPUPROFILER_EVENT_SCOPE(UAjaMediaCapture::OnFrameCaptured_RenderingThread::UnlockDMATexture);
-	TextureTransfer->UnlockTexture(InTexture->GetTexture2D());
+	if (ShouldCaptureRHIResource())
+	{
+		TRACE_CPUPROFILER_EVENT_SCOPE(UAjaMediaCapture::OnFrameCaptured_RenderingThread::UnlockDMATexture);
+		TextureTransfer->UnlockTexture(InTexture->GetTexture2D());
+	}
 }
 
 void UAjaMediaCapture::OnFrameCaptured_AnyThread(const FCaptureBaseData& InBaseData, TSharedPtr<FMediaCaptureUserData, ESPMode::ThreadSafe> InUserData, const FMediaCaptureResourceData& InResourceData)
@@ -611,13 +635,13 @@ void UAjaMediaCapture::WaitForSync_AnyThread() const
 
 void UAjaMediaCapture::OutputAudio_AnyThread(const AJA::AJAOutputFrameBufferData& FrameBuffer) const
 {
-	if (bOutputAudio)
+	if (bOutputAudio && !bDirectlyWriteAudio)
 	{
 		TRACE_CPUPROFILER_EVENT_SCOPE(UAjaMediaCapture::OnFrameCaptured_RenderingThread::SetAudio);
 		
 		// Take a local copy of the audio output in case it is switched from the main thread.
 		const TSharedPtr<FMediaIOAudioOutput> LocalAudioOutput = AudioOutput;
-		if(!LocalAudioOutput)
+		if (!LocalAudioOutput)
 		{
 			return;
 		}
@@ -788,3 +812,17 @@ void UAjaMediaCapture::OnFrameCapturedInternal_AnyThread(const FCaptureBaseData&
 		SetState(EMediaCaptureState::Error);
 	}
 }
+
+void UAjaMediaCapture::OnAudioBufferReceived_AudioThread(Audio::FDeviceId DeviceId, float* Data, int32 NumSamples) const
+{
+	TRACE_CPUPROFILER_EVENT_SCOPE(UAjaMediaCapture::AudioBufferReceived);
+
+	if (bOutputAudio && OutputChannel)
+	{
+		if (ensure(NumInputChannels != 0))
+		{
+			const TArray<int32> ConvertedSamples = FMediaIOAudioOutput::ConvertAndUpmixBuffer<int32>(MakeArrayView<float>(Data, NumSamples), NumInputChannels, NumOutputChannels);
+			OutputChannel->DMAWriteAudio(reinterpret_cast<const uint8*>(ConvertedSamples.GetData()), ConvertedSamples.Num() * sizeof(int32));
+		}
+	}
+}
diff --git a/Engine/Plugins/Media/AjaMedia/Source/AjaMediaOutput/Public/AjaMediaCapture.h b/Engine/Plugins/Media/AjaMedia/Source/AjaMediaOutput/Public/AjaMediaCapture.h
--- a/Engine/Plugins/Media/AjaMedia/Source/AjaMediaOutput/Public/AjaMediaCapture.h
+++ b/Engine/Plugins/Media/AjaMedia/Source/AjaMediaOutput/Public/AjaMediaCapture.h
@@ -78,6 +78,7 @@ private:
 	bool CleanupPreEditorExit();
 	void OnEnginePreExit();
 	void OnFrameCapturedInternal_AnyThread(const FCaptureBaseData& InBaseData, TSharedPtr<FMediaCaptureUserData, ESPMode::ThreadSafe> InUserData, const FMediaCaptureResourceData& InResourceData);
+	void OnAudioBufferReceived_AudioThread(Audio::FDeviceId DeviceId, float* Data, int32 NumSamples) const;
 
 	struct FAudioBuffer
 	{
@@ -116,7 +117,16 @@ private:
 	/** Holds an audio output that captures audio from the engine. */
 	TSharedPtr<class FMediaIOAudioOutput> AudioOutput;
 
+	/** Number of audio channels output to the AJA card.  */
+	int32 NumOutputChannels = 0;
+
+	/** Number of audio channels the engine is rendering to. */
+	int32 NumInputChannels =0;
+
 	bool bOutputAudio = false;
+
+	/** Whether to write audio on the audio thread instead of collecting samples and outputting on the render thread.. */
+	bool bDirectlyWriteAudio = false;
 	
 	/** Textures to release when the capture has stopped. Must be released after GPUTextureTransfer textures have been unregistered. */
 	TArray<FTextureRHIRef> TexturesToRelease;
diff --git a/Engine/Plugins/Media/AjaMedia/Source/AjaMediaOutput/Public/AjaMediaOutput.h b/Engine/Plugins/Media/AjaMedia/Source/AjaMediaOutput/Public/AjaMediaOutput.h
--- a/Engine/Plugins/Media/AjaMedia/Source/AjaMediaOutput/Public/AjaMediaOutput.h
+++ b/Engine/Plugins/Media/AjaMedia/Source/AjaMediaOutput/Public/AjaMediaOutput.h
@@ -91,6 +91,10 @@ public:
 	UPROPERTY(BlueprintReadWrite, EditAnywhere, Category = "Output")
 	EAjaMediaOutputAudioSampleRate AudioSampleRate = EAjaMediaOutputAudioSampleRate::SR_48k;
 
+	/** Whether to output audio on the audio thread instead of the rendering thread. Only available in Ping Pong mode. (Experimental) */
+    UPROPERTY(BlueprintReadWrite, EditAnywhere, Category = "Output", meta = (EditCondition = "bOutputWithAutoCirculating == false"))
+	bool bOutputAudioOnAudioThread = false;
+
 	/**
 	 * Number of frame used to transfer from the system memory to the AJA card.
 	 * A smaller number is most likely to cause missed frame.
diff --git a/Engine/Plugins/Media/MediaIOFramework/Source/MediaIOCore/Private/MediaIOCoreAudioOutput.cpp b/Engine/Plugins/Media/MediaIOFramework/Source/MediaIOCore/Private/MediaIOCoreAudioOutput.cpp
--- a/Engine/Plugins/Media/MediaIOFramework/Source/MediaIOCore/Private/MediaIOCoreAudioOutput.cpp
+++ b/Engine/Plugins/Media/MediaIOFramework/Source/MediaIOCore/Private/MediaIOCoreAudioOutput.cpp
@@ -73,6 +73,8 @@ void FMediaIOAudioCapture::OnNewSubmixBuffer(const USoundSubmix* InOwningSubmix,
 	{
 		if (ensureMsgf(NumChannels == InNumChannels, TEXT("Expected %d channels from submix buffer but got %d instead."), NumChannels, InNumChannels))
 		{
+			AudioCapturedDelegate.ExecuteIfBound(InAudioData, InNumSamples);
+			
 			int32 NumPushed = AudioSplitter.PushAudio(InAudioData, InNumSamples);
 			if (InNumSamples != NumPushed && NumPushed != -1)
 			{
diff --git a/Engine/Plugins/Media/MediaIOFramework/Source/MediaIOCore/Private/MediaIOCoreSubsystem.cpp b/Engine/Plugins/Media/MediaIOFramework/Source/MediaIOCore/Private/MediaIOCoreSubsystem.cpp
--- a/Engine/Plugins/Media/MediaIOFramework/Source/MediaIOCore/Private/MediaIOCoreSubsystem.cpp
+++ b/Engine/Plugins/Media/MediaIOFramework/Source/MediaIOCore/Private/MediaIOCoreSubsystem.cpp
@@ -28,6 +28,7 @@ TSharedPtr<FMediaIOAudioOutput> UMediaIOCoreSubsystem::CreateAudioOutput(const F
 		else
 		{
 			MediaIOAudioCapture = MediaIOAudioCaptures.Add(InArgs.AudioDeviceHandle.GetDeviceID(), MakeUnique<FMediaIOAudioCapture>(InArgs.AudioDeviceHandle)).Get();
+			MainMediaIOAudioCapture->OnAudioCaptured_RenderThread().BindUObject(this, &UMediaIOCoreSubsystem::OnBufferReceivedByCapture, InArgs.AudioDeviceHandle.GetDeviceID());
 		}
 	}
 	else
@@ -36,6 +37,7 @@ TSharedPtr<FMediaIOAudioOutput> UMediaIOCoreSubsystem::CreateAudioOutput(const F
 		if (!MainMediaIOAudioCapture)
 		{
 			MainMediaIOAudioCapture = MakeUnique<FMainMediaIOAudioCapture>();
+			MainMediaIOAudioCapture->OnAudioCaptured_RenderThread().BindUObject(this, &UMediaIOCoreSubsystem::OnBufferReceivedByCapture, InArgs.AudioDeviceHandle.GetDeviceID());
 		}
 		MediaIOAudioCapture = MainMediaIOAudioCapture.Get();
 	}	
@@ -43,6 +45,15 @@ TSharedPtr<FMediaIOAudioOutput> UMediaIOCoreSubsystem::CreateAudioOutput(const F
 	return MediaIOAudioCapture->CreateAudioOutput(InArgs.NumOutputChannels, InArgs.TargetFrameRate, InArgs.MaxSampleLatency, InArgs.OutputSampleRate);
 }
 
+int32 UMediaIOCoreSubsystem::GetNumAudioInputChannels() const
+{
+	if (MainMediaIOAudioCapture)
+	{
+		return MainMediaIOAudioCapture->GetNumInputChannels();
+	}
+	return 0;
+}
+
 void UMediaIOCoreSubsystem::OnAudioDeviceDestroyed(Audio::FDeviceId InAudioDeviceId)
 {
 	if (MediaIOAudioCaptures.Contains(InAudioDeviceId))
@@ -50,3 +61,8 @@ void UMediaIOCoreSubsystem::OnAudioDeviceDestroyed(Audio::FDeviceId InAudioDevic
 		MediaIOAudioCaptures.Remove(InAudioDeviceId);
 	}
 }
+
+void UMediaIOCoreSubsystem::OnBufferReceivedByCapture(float* Data, int32 NumSamples, Audio::FDeviceId AudioDeviceID) const
+{
+	BufferReceivedDelegate.Broadcast(AudioDeviceID, Data, NumSamples);
+}
diff --git a/Engine/Plugins/Media/MediaIOFramework/Source/MediaIOCore/Public/MediaIOCoreAudioOutput.h b/Engine/Plugins/Media/MediaIOFramework/Source/MediaIOCore/Public/MediaIOCoreAudioOutput.h
--- a/Engine/Plugins/Media/MediaIOFramework/Source/MediaIOCore/Public/MediaIOCoreAudioOutput.h
+++ b/Engine/Plugins/Media/MediaIOFramework/Source/MediaIOCore/Public/MediaIOCoreAudioOutput.h
@@ -13,6 +13,8 @@
 class FAudioDevice;
 class FAudioDeviceHandle;
 
+DECLARE_DELEGATE_TwoParams(FOnBufferReceived, uint8* Buffer, int32 BufferSize);
+
 class MEDIAIOCORE_API FMediaIOAudioOutput
 { 
 public:
@@ -24,11 +26,11 @@ public:
 		uint32 InMaxSampleLatency = 0;
 		uint32 InOutputSampleRate = 0;
 	};
-
+	
 	FMediaIOAudioOutput(Audio::FPatchOutputStrongPtr InPatchOutput, const FAudioOptions& InAudioOptions);
-
+	
 	template <typename OutputType>
-	TArray<OutputType> ConvertAndUpmixBuffer(const Audio::FAlignedFloatBuffer& InBuffer, int32 InNumInputChannels, int32 InNumOutputChannels) const
+	static TArray<OutputType> ConvertAndUpmixBuffer(TConstArrayView<float> InBuffer, int32 InNumInputChannels, int32 InNumOutputChannels)
 	{
 		TRACE_CPUPROFILER_EVENT_SCOPE(MediaIoCoreModule::ConvertAndUpmixBuffer);
 
@@ -110,6 +112,10 @@ public:
 		return ConvertAndUpmixBuffer<OutputType>(FloatBuffer, NumInputChannels, NumOutputChannels);
 	}
 
+	FOnBufferReceived& OnBufferReceived()
+	{
+		return BufferReceivedDelegate;
+	}
 
 public:
 	// @todo: Depend on frame number to correctly fetch the right amount of frames on framerates like 59.97
@@ -136,6 +142,8 @@ private:
 	/** The buffer accumulating audio samples. */
 	Audio::FPatchOutputStrongPtr PatchOutput;
 
+	/** Callback called when a buffer is received. */
+	FOnBufferReceived BufferReceivedDelegate;
 };
 
 /**
@@ -144,6 +152,8 @@ private:
 class FMediaIOAudioCapture : public ISubmixBufferListener
 {
 public:
+	DECLARE_DELEGATE_TwoParams(FOnAudioCaptured, float* /* data */, int32 /* NumSamples */);
+
 	FMediaIOAudioCapture(const FAudioDeviceHandle& InAudioDeviceHandle);
 	virtual ~FMediaIOAudioCapture();
 
@@ -153,6 +163,16 @@ public:
 	/** Create an audio output that will receive audio samples. */
 	TSharedPtr<FMediaIOAudioOutput> CreateAudioOutput(int32 InNumOutputChannels, FFrameRate InTargetFrameRate, uint32 InMaxSampleLatency, uint32 InOutputSampleRate);
 
+	FOnAudioCaptured& OnAudioCaptured_RenderThread()
+	{
+		return AudioCapturedDelegate;
+	}
+
+	int32 GetNumInputChannels() const
+	{
+		return NumChannels;
+	}
+
 protected:
 	void RegisterBufferListener(FAudioDevice* AudioDevice);
 	void UnregisterBufferListener(FAudioDevice* AudioDevice);
@@ -176,6 +196,9 @@ private:
 
 	/** Utility that allows pushing audio samples to multiple outputs. */
 	Audio::FPatchSplitter AudioSplitter;
+
+	/** Callback for a audio rendered event. */
+	FOnAudioCaptured AudioCapturedDelegate;
 };
 
 /**
diff --git a/Engine/Plugins/Media/MediaIOFramework/Source/MediaIOCore/Public/MediaIOCoreSubsystem.h b/Engine/Plugins/Media/MediaIOFramework/Source/MediaIOCore/Public/MediaIOCoreSubsystem.h
--- a/Engine/Plugins/Media/MediaIOFramework/Source/MediaIOCore/Public/MediaIOCoreSubsystem.h
+++ b/Engine/Plugins/Media/MediaIOFramework/Source/MediaIOCore/Public/MediaIOCoreSubsystem.h
@@ -2,6 +2,8 @@
 
 #pragma once
 
+#include "AudioDeviceHandle.h"
+
 #include "AudioDeviceHandle.h"
 #include "MediaIOCoreAudioOutput.h"
 #include "Subsystems/EngineSubsystem.h"
@@ -12,6 +14,8 @@ UCLASS()
 class MEDIAIOCORE_API UMediaIOCoreSubsystem : public UEngineSubsystem
 {
 public:
+	DECLARE_MULTICAST_DELEGATE_ThreeParams(FOnBufferReceived, Audio::FDeviceId /* DeviceId */, float* /* Data */, int32 /* NumSamples */)
+	
 	struct FCreateAudioOutputArgs
 	{
 		uint32 NumOutputChannels = 0;
@@ -34,8 +38,22 @@ public:
 	 */
 	TSharedPtr<FMediaIOAudioOutput> CreateAudioOutput(const FCreateAudioOutputArgs& InArgs);
 
+	/**
+	 * Get the number of audio channels used by the main audio device.
+	 **/
+	int32 GetNumAudioInputChannels() const;
+
+	/**
+	 * @Note: Called from the audio thread.
+	 */
+	FOnBufferReceived& OnBufferReceived_AudioThread()
+	{
+		return BufferReceivedDelegate;
+	}
+
 private:
 	void OnAudioDeviceDestroyed(Audio::FDeviceId InAudioDeviceId);
+	void OnBufferReceivedByCapture(float* Data, int32 NumSamples, Audio::FDeviceId AudioDeviceID) const;
 
 private:
 	TUniquePtr<FMediaIOAudioCapture> MainMediaIOAudioCapture;
@@ -43,4 +61,6 @@ private:
 	TMap<Audio::FDeviceId, TUniquePtr<FMediaIOAudioCapture>> MediaIOAudioCaptures;
 
 	FDelegateHandle DeviceDestroyedHandle;
+
+	FOnBufferReceived BufferReceivedDelegate;
 };
