Fix missing include path that is preventing distribution of AjaCore

#rnx

diff --git a/Engine/Plugins/Media/AjaMedia/Source/Aja/AjaCore.Build.cs b/Engine/Plugins/Media/AjaMedia/Source/Aja/AjaCore.Build.cs
--- a/Engine/Plugins/Media/AjaMedia/Source/Aja/AjaCore.Build.cs
+++ b/Engine/Plugins/Media/AjaMedia/Source/Aja/AjaCore.Build.cs
@@ -21,6 +21,7 @@ public class AjaCore : ModuleRules
 		PrivateIncludePaths.AddRange(
 			new string[] {
 				Path.Combine(AjaThirdPartyPath, "ntv2lib-deploy", "includes"),
+				Path.Combine(AjaThirdPartyPath, "ntv2lib-deploy", "includes", "ajantv2", "includes"),
 				Path.Combine(AjaThirdPartyPath, "ntv2lib-deploy", "includes", "ajantv2", "src", "win")
 			});
 			
