Aja - Add options for HDR support

diff --git a/Engine/Plugins/Media/AjaMedia/Source/Aja/Private/ChannelBase.cpp b/Engine/Plugins/Media/AjaMedia/Source/Aja/Private/ChannelBase.cpp
--- a/Engine/Plugins/Media/AjaMedia/Source/Aja/Private/ChannelBase.cpp
+++ b/Engine/Plugins/Media/AjaMedia/Source/Aja/Private/ChannelBase.cpp
@@ -311,7 +311,7 @@ namespace AJA
 				DesiredVideoFormat = Helpers::GetLevelA(DesiredVideoFormat);
 			}
 
-			if (!InCommandList.RegisterChannel(GetOptions().TransportType, InputSource, Channel, IsInput(), bConnectChannel, GetOptions().TimecodeFormat, GetOptions().PixelFormat, DesiredVideoFormat, bAsOwner, Options.bAutoDetectFormat))
+			if (!InCommandList.RegisterChannel(GetOptions().TransportType, InputSource, Channel, IsInput(), bConnectChannel, GetOptions().TimecodeFormat, GetOptions().PixelFormat, DesiredVideoFormat, Options.HDROptions, bAsOwner, Options.bAutoDetectFormat))
 			{
 				return false;
 			}
@@ -439,9 +439,14 @@ namespace AJA
 				}
 
 				const int32_t NumberOfLinkChannel = Helpers::GetNumberOfLinkChannel(GetOptions().TransportType);
+
+				const NTV2HDRXferChars ETOF = Helpers::ConvertToAjaHDRXferChars(GetOptions().HDROptions.EOTF);
+				const NTV2HDRColorimetry Colorimetry = Helpers::ConvertToAjaHDRColorimetry(GetOptions().HDROptions.Gamut);
+				const NTV2HDRLuminance Luminance = Helpers::ConvertToAjaHDRLuminance(GetOptions().HDROptions.Luminance);
+				
 				for (int32_t ChannelIndex = 0; ChannelIndex < NumberOfLinkChannel; ++ChannelIndex)
 				{
-					AJA_CHECK(Device->GetCard()->SetFrameBufferFormat(NTV2Channel(int32_t(Channel) + ChannelIndex), PixelFormat));
+					AJA_CHECK(Device->GetCard()->SetFrameBufferFormat(NTV2Channel(int32_t(Channel) + ChannelIndex), PixelFormat, AJA_RETAIL_DEFAULT, ETOF, Colorimetry, Luminance));
 				}
 
 				// if output, set the Reference type
@@ -543,7 +548,7 @@ namespace AJA
 					{
 						bConnectedChannel = false;
 						const bool bAsOwner = true;
-						if (!InCommandList.RegisterChannel(GetOptions().TransportType, KeySource, KeyChannel, IsInput(), bConnectedChannel, ETimecodeFormat::TCF_None, GetOptions().PixelFormat, DesiredVideoFormat, bAsOwner, Options.bAutoDetectFormat))
+						if (!InCommandList.RegisterChannel(GetOptions().TransportType, KeySource, KeyChannel, IsInput(), bConnectedChannel, ETimecodeFormat::TCF_None, GetOptions().PixelFormat, DesiredVideoFormat, Options.HDROptions, bAsOwner, Options.bAutoDetectFormat))
 						{
 							return false;
 						}
@@ -584,7 +589,7 @@ namespace AJA
 					{
 						const bool bUseKeyConnectedChannel = false;
 						const bool bAsOwner = true;
-						if (!InCommandList.RegisterChannel(GetOptions().TransportType, KeySource, KeyChannel, IsInput(), bUseKeyConnectedChannel, ETimecodeFormat::TCF_None, GetOptions().PixelFormat, DesiredVideoFormat, bAsOwner, Options.bAutoDetectFormat))
+						if (!InCommandList.RegisterChannel(GetOptions().TransportType, KeySource, KeyChannel, IsInput(), bUseKeyConnectedChannel, ETimecodeFormat::TCF_None, GetOptions().PixelFormat, DesiredVideoFormat, Options.HDROptions, bAsOwner, Options.bAutoDetectFormat))
 						{
 							return false;
 						}
diff --git a/Engine/Plugins/Media/AjaMedia/Source/Aja/Private/Device.cpp b/Engine/Plugins/Media/AjaMedia/Source/Aja/Private/Device.cpp
--- a/Engine/Plugins/Media/AjaMedia/Source/Aja/Private/Device.cpp
+++ b/Engine/Plugins/Media/AjaMedia/Source/Aja/Private/Device.cpp
@@ -162,7 +162,13 @@ namespace AJA
 			: Connection(InConnection)
 		{}
 
-		bool DeviceConnection::CommandList::RegisterChannel(ETransportType InTransportType, NTV2InputSource InInputSource, NTV2Channel InChannel, bool bInIsInput, bool bConnectChannel, ETimecodeFormat InTimecodeFormat, EPixelFormat InUEPixelFormat, NTV2VideoFormat InDesiredInputFormat, bool bInAsOwner, bool bInIsAutoDetected)
+		bool DeviceConnection::CommandList::RegisterChannel(ETransportType InTransportType, NTV2InputSource InInputSource, NTV2Channel InChannel, bool bInAsInput, bool bConnectChannel, ETimecodeFormat InTimecodeFormat, EPixelFormat InUEPixelFormat, NTV2VideoFormat InDesiredInputFormat, bool bInAsOwner, bool bInIsAutoDetected)
+		{
+			FAjaHDROptions UnusedOptions;
+			return RegisterChannel(InTransportType, InInputSource, InChannel, bInAsInput, bConnectChannel, InTimecodeFormat, InUEPixelFormat, InDesiredInputFormat, UnusedOptions, bInAsOwner, bInIsAutoDetected);
+		}
+
+		bool DeviceConnection::CommandList::RegisterChannel(ETransportType InTransportType, NTV2InputSource InInputSource, NTV2Channel InChannel, bool bInIsInput, bool bConnectChannel, ETimecodeFormat InTimecodeFormat, EPixelFormat InUEPixelFormat, NTV2VideoFormat InDesiredInputFormat, FAjaHDROptions& InOutHDROptions, bool bInAsOwner, bool bInIsAutoDetected)
 		{
 			AJAAutoLock Lock(&Connection.ChannelLock);
 
@@ -185,7 +191,7 @@ namespace AJA
 
 				if (bResult)
 				{
-					bResult = Connection.Lock_EnableChannel_Helper(Connection.Card, InTransportType, InInputSource, InChannel, bInIsInput, bConnectChannel, bInIsAutoDetected, InTimecodeFormat, InUEPixelFormat, InDesiredInputFormat, VideoFormat);
+					bResult = Connection.Lock_EnableChannel_Helper(Connection.Card, InTransportType, InInputSource, InChannel, bInIsInput, bConnectChannel, bInIsAutoDetected, InTimecodeFormat, InUEPixelFormat, InDesiredInputFormat, VideoFormat, InOutHDROptions);
 				}
 
 				uint32_t NewBaseFrameBufferIndex = InvalidFrameBufferIndex;
@@ -242,7 +248,7 @@ namespace AJA
 
 					if (bResult)
 					{
-						bResult = Connection.Lock_EnableChannel_Helper(Connection.Card, InTransportType, InInputSource, InChannel, bInIsInput, bConnectChannel, bInIsAutoDetected, InTimecodeFormat, InUEPixelFormat, InDesiredInputFormat, VideoFormat);
+						bResult = Connection.Lock_EnableChannel_Helper(Connection.Card, InTransportType, InInputSource, InChannel, bInIsInput, bConnectChannel, bInIsAutoDetected, InTimecodeFormat, InUEPixelFormat, InDesiredInputFormat, VideoFormat, InOutHDROptions);
 					}
 				}
 
@@ -958,7 +964,7 @@ namespace AJA
 			return true;
 		}
 
-		bool DeviceConnection::Lock_EnableChannel_Helper(CNTV2Card* InCard, ETransportType InTransportType, NTV2InputSource InInputSource, NTV2Channel InChannel, bool bIsInput, bool bConnectChannel, bool bIsAutoDetected, ETimecodeFormat InTimecodeFormat, EPixelFormat InUEPixelFormat, NTV2VideoFormat InDesiredInputFormat, NTV2VideoFormat& OutFoundVideoFormat)
+		bool DeviceConnection::Lock_EnableChannel_Helper(CNTV2Card* InCard, ETransportType InTransportType, NTV2InputSource InInputSource, NTV2Channel InChannel, bool bIsInput, bool bConnectChannel, bool bIsAutoDetected, ETimecodeFormat InTimecodeFormat, EPixelFormat InUEPixelFormat, NTV2VideoFormat InDesiredInputFormat, NTV2VideoFormat& OutFoundVideoFormat, FAjaHDROptions& InOutHDROptions)
 		{
 			NTV2DeviceID DeviceId = Card->GetDeviceID();
 			if (InChannel >= ::NTV2DeviceGetNumFrameStores(DeviceId))
@@ -1049,6 +1055,12 @@ namespace AJA
 					UE_LOG(LogTemp, Error, TEXT("Device: Initialization of the input failed for channel %d on device '%S'. %S"), uint32_t(InChannel) + 1, InCard->GetDisplayName().c_str(), FailureReason.c_str());
 					return false;
 				}
+				
+				if (!Helpers::GetInputHDRMetadata(InCard, InChannel, InOutHDROptions))
+				{
+					UE_LOG(LogTemp, Error, TEXT("Device: Could not fetch HDR metadata from channel %d on device '%S'."), uint32_t(InChannel) + 1, InCard->GetDisplayName().c_str());
+					return false;
+				}
 			}
 
 			// Setup basic video to support Interrupt. This will be override by ChannelBase if necessary
@@ -1072,9 +1084,13 @@ namespace AJA
 			const NTV2FrameBufferFormat FrameBufferFormat = Helpers::ConvertPixelFormatToFrameBufferFormat(InUEPixelFormat);
 			AJA_CHECK(InCard->SetVideoFormat(OutFoundVideoFormat, AJA_RETAIL_DEFAULT, false, InChannel));
 
+			const NTV2HDRXferChars EOTF = Helpers::ConvertToAjaHDRXferChars(InOutHDROptions.EOTF);
+			const NTV2HDRColorimetry Colorimetry = Helpers::ConvertToAjaHDRColorimetry(InOutHDROptions.Gamut);
+			const NTV2HDRLuminance Luminance = Helpers::ConvertToAjaHDRLuminance(InOutHDROptions.Luminance);
+			
 			for (int32_t ChannelIndex = 0; ChannelIndex < NumberOfLinkChannel; ++ChannelIndex)
 			{
-				AJA_CHECK(InCard->SetFrameBufferFormat(NTV2Channel(int32_t(InChannel) + ChannelIndex), FrameBufferFormat));
+				AJA_CHECK(InCard->SetFrameBufferFormat(NTV2Channel(int32_t(InChannel) + ChannelIndex), FrameBufferFormat, AJA_RETAIL_DEFAULT, EOTF, Colorimetry, Luminance));
 			}
 
 			// Route the input for the sync, seems required
diff --git a/Engine/Plugins/Media/AjaMedia/Source/Aja/Private/Device.h b/Engine/Plugins/Media/AjaMedia/Source/Aja/Private/Device.h
--- a/Engine/Plugins/Media/AjaMedia/Source/Aja/Private/Device.h
+++ b/Engine/Plugins/Media/AjaMedia/Source/Aja/Private/Device.h
@@ -121,6 +121,8 @@ namespace AJA
 
 			public:
 				bool RegisterChannel(ETransportType InTransportType, NTV2InputSource InInputSource, NTV2Channel InChannel, bool bInAsInput, bool bConnectChannel, ETimecodeFormat InTimecodeFormat, EPixelFormat InUEPixelFormat, NTV2VideoFormat InDesiredInputFormat, bool bInAsOwner, bool bInIsAutoDetected);
+				/** RegisterChannel override that handles HDR metadata. */
+				bool RegisterChannel(ETransportType InTransportType, NTV2InputSource InInputSource, NTV2Channel InChannel, bool bInAsInput, bool bConnectChannel, ETimecodeFormat InTimecodeFormat, EPixelFormat InUEPixelFormat, NTV2VideoFormat InDesiredInputFormat, FAjaHDROptions& InOutHDROptions, bool bInAsOwner, bool bInIsAutoDetected);
 				bool UnregisterChannel(NTV2Channel InChannel, bool bInAsOwner);
 
 				bool IsChannelConnect(NTV2Channel InChannel);
@@ -144,7 +146,7 @@ namespace AJA
 			void Lock_UnInitialize();
 
 			bool Lock_EnableChannel_SDILinkHelper(CNTV2Card* InCard, ETransportType InTransportType, NTV2Channel InChannel) const;
-			bool Lock_EnableChannel_Helper(CNTV2Card* InCard, ETransportType InTransportType, NTV2InputSource InInputSource, NTV2Channel InChannel, bool InAsInput, bool bConnectChannel, bool bIsAutoDetected, ETimecodeFormat InTimecodeFormat, EPixelFormat InUEPixelFormat, NTV2VideoFormat InDesiredInputFormat, NTV2VideoFormat& OutFoundVideoFormat);
+			bool Lock_EnableChannel_Helper(CNTV2Card* InCard, ETransportType InTransportType, NTV2InputSource InInputSource, NTV2Channel InChannel, bool InAsInput, bool bConnectChannel, bool bIsAutoDetected, ETimecodeFormat InTimecodeFormat, EPixelFormat InUEPixelFormat, NTV2VideoFormat InDesiredInputFormat, NTV2VideoFormat& OutFoundVideoFormat, FAjaHDROptions& InOutHDROptions);
 			void Lock_SubscribeChannel_Helper(CNTV2Card* InCard, NTV2Channel InChannel, bool InAsInput) const;
 
 			uint32_t Lock_AcquireBaseFrameIndex(NTV2Channel InChannel, NTV2VideoFormat InVideoFormat) const;
diff --git a/Engine/Plugins/Media/AjaMedia/Source/Aja/Private/Helpers.cpp b/Engine/Plugins/Media/AjaMedia/Source/Aja/Private/Helpers.cpp
--- a/Engine/Plugins/Media/AjaMedia/Source/Aja/Private/Helpers.cpp
+++ b/Engine/Plugins/Media/AjaMedia/Source/Aja/Private/Helpers.cpp
@@ -2,7 +2,6 @@
 
 #include "Helpers.h"
 
-
 namespace AJA
 {
 	namespace Private
@@ -130,7 +129,6 @@ namespace AJA
 
 		AJA_PixelFormat Helpers::ConvertToPixelFormat(EPixelFormat InPixelFormat)
 		{
-			NTV2FrameBufferFormat Result = NTV2_FBF_ARGB;
 			switch (InPixelFormat)
 			{
 			case EPixelFormat::PF_8BIT_YCBCR:
@@ -145,6 +143,90 @@ namespace AJA
 			}
 		}
 
+		NTV2HDRXferChars Helpers::ConvertToAjaHDRXferChars(EAjaHDRMetadataEOTF HDRXferChars)
+		{
+			switch(HDRXferChars)
+			{
+			case EAjaHDRMetadataEOTF::PQ:
+				return NTV2_VPID_TC_PQ;
+			case EAjaHDRMetadataEOTF::HLG:
+				return NTV2_VPID_TC_HLG;
+			case EAjaHDRMetadataEOTF::SDR:
+				return NTV2_VPID_TC_SDR_TV;
+			default:
+				return NTV2_VPID_TC_Unspecified;
+			}
+		}
+		
+		EAjaHDRMetadataEOTF Helpers::ConvertFromAjaHDRXferChars(NTV2HDRXferChars HDRXferChars)
+		{
+			switch(HDRXferChars)
+			{
+			case NTV2_VPID_TC_PQ:
+				return EAjaHDRMetadataEOTF::PQ;
+			case NTV2_VPID_TC_HLG:
+				return EAjaHDRMetadataEOTF::HLG;
+			case NTV2_VPID_TC_SDR_TV:
+				return EAjaHDRMetadataEOTF::SDR;
+			default:
+				return EAjaHDRMetadataEOTF::Unspecified;
+			}
+		}
+
+		NTV2HDRColorimetry Helpers::ConvertToAjaHDRColorimetry(EAjaHDRMetadataGamut HDRColorimetry)
+		{
+			switch (HDRColorimetry)
+			{
+			case EAjaHDRMetadataGamut::Rec709:
+				return NTV2_VPID_Color_Rec709;
+			case EAjaHDRMetadataGamut::Rec2020:
+				return NTV2_VPID_Color_UHDTV;
+			default:
+				return NTV2_VPID_Color_Unknown;
+			}
+		}
+
+		EAjaHDRMetadataGamut Helpers::ConvertFromAjaHDRColorimetry(NTV2HDRColorimetry HDRColorimetry)
+		{
+			switch (HDRColorimetry)
+			{
+			case NTV2_VPID_Color_Rec709:
+				return EAjaHDRMetadataGamut::Rec709;
+			case NTV2_VPID_Color_UHDTV:
+				return EAjaHDRMetadataGamut::Rec2020;
+			default:
+				return EAjaHDRMetadataGamut::Invalid;
+			}
+		}
+		
+		NTV2HDRLuminance Helpers::ConvertToAjaHDRLuminance(EAjaHDRMetadataLuminance HDRLuminance)
+		{
+			switch (HDRLuminance)
+			{
+			case EAjaHDRMetadataLuminance::ICtCp:
+				return NTV2_VPID_Luminance_ICtCp;
+			case EAjaHDRMetadataLuminance::YCbCr:
+				return NTV2_VPID_Luminance_YCbCr;
+			default:
+				checkNoEntry();
+				return NTV2_VPID_Luminance_YCbCr;
+			}
+		}
+		
+		EAjaHDRMetadataLuminance Helpers::ConvertFromAjaHDRLuminance(NTV2HDRLuminance HDRLuminance)
+		{
+			switch (HDRLuminance)
+			{
+			case NTV2_VPID_Luminance_ICtCp:
+				return EAjaHDRMetadataLuminance::ICtCp;
+			case NTV2_VPID_Luminance_YCbCr:
+				return EAjaHDRMetadataLuminance::YCbCr;
+			default:
+				checkNoEntry();
+				return EAjaHDRMetadataLuminance::YCbCr;
+			}
+		}
+
 		NTV2FrameRate Helpers::ConvertToFrameRate(uint32_t InNumerator, uint32_t InDenominator)
 		{
 			if (InNumerator == 120    && InDenominator == 1)		{ return NTV2_FRAMERATE_12000; }
@@ -394,6 +476,28 @@ namespace AJA
 			return false;
 		}
 
+		bool Helpers::GetInputHDRMetadata(CNTV2Card* InCard, NTV2Channel InChannel, FAjaHDROptions& OutHDRMetadata)
+		{
+			NTV2VPIDTransferCharacteristics EOTF = NTV2VPIDTransferCharacteristics::NTV2_VPID_TC_SDR_TV;
+			NTV2VPIDColorimetry Colorimetry = NTV2VPIDColorimetry::NTV2_VPID_Color_Rec709;
+			NTV2VPIDLuminance Luminance = NTV2VPIDLuminance::NTV2_VPID_Luminance_YCbCr;
+			
+			const bool bSuccess = InCard->GetVPIDTransferCharacteristics(EOTF, InChannel)
+				&& InCard->GetVPIDColorimetry(Colorimetry, InChannel)
+				&& InCard->GetVPIDLuminance(Luminance, InChannel);
+
+			if (!bSuccess)
+			{
+				return false;
+			}
+			
+			OutHDRMetadata.EOTF = ConvertFromAjaHDRXferChars(EOTF);
+			OutHDRMetadata.Gamut = ConvertFromAjaHDRColorimetry(Colorimetry);
+			OutHDRMetadata.Luminance = ConvertFromAjaHDRLuminance(Luminance);
+			
+			return true;
+		}
+
 		bool Helpers::SetSDIOutLevelAtoLevelBConversion(CNTV2Card* InCard, ETransportType InTransportType, NTV2Channel InChannel, NTV2VideoFormat InFormat, bool bValue)
 		{
 			const int32_t NumberOfLinkChannel = Helpers::GetNumberOfLinkChannel(InTransportType);
@@ -409,7 +513,7 @@ namespace AJA
 
 			return bDoLevelConversion;
 		}
-
+		
 		void Helpers::SetSDIInLevelBtoLevelAConversion(CNTV2Card* InCard, ETransportType InTransportType, NTV2Channel InChannel, NTV2VideoFormat InFormat, NTV2VideoFormat& OutFormat)
 		{
 			const int32_t NumberOfLinkChannel = Helpers::GetNumberOfLinkChannel(InTransportType);
diff --git a/Engine/Plugins/Media/AjaMedia/Source/Aja/Private/Helpers.h b/Engine/Plugins/Media/AjaMedia/Source/Aja/Private/Helpers.h
--- a/Engine/Plugins/Media/AjaMedia/Source/Aja/Private/Helpers.h
+++ b/Engine/Plugins/Media/AjaMedia/Source/Aja/Private/Helpers.h
@@ -10,7 +10,7 @@
 
 THIRD_PARTY_INCLUDES_START
 
-__pragma(warning(disable: 4263))  /* Member function does not override any base class virtual member function. */ 
+__pragma(warning(disable: 4263))  /* Member function does not override any base class virtual member function. */
 __pragma(warning(disable: 4264))  /* No override available for virtual member function from base. */
 
 #if PLATFORM_WINDOWS
@@ -46,7 +46,6 @@ __pragma(warning(disable: 4264))  /* No override available for virtual member fu
 
 THIRD_PARTY_INCLUDES_END
 
-
 #include "AJALib.h"
 #include <optional>
 
@@ -93,11 +92,27 @@ namespace AJA
 			static NTV2FrameBufferFormat ConvertPixelFormatToFrameBufferFormat(EPixelFormat InPixelFormat);
 			static EPixelFormat ConvertFrameBufferFormatToPixelFormat(NTV2FrameBufferFormat InPixelFormat);
 			static AJA_PixelFormat ConvertToPixelFormat(EPixelFormat InPixelFormat);
+			
+			/** Convert HDR transfer characteristics from Unreal to Aja format. */
+			static NTV2HDRXferChars ConvertToAjaHDRXferChars(EAjaHDRMetadataEOTF HDRXferChars);
+			/** Convert HDR transfer characteristics from Aja to Unreal format. */
+			static EAjaHDRMetadataEOTF ConvertFromAjaHDRXferChars(NTV2HDRXferChars XferChar);
+			/** Convert HDR colorimetry from Unreal to Aja format. */
+			static NTV2HDRColorimetry ConvertToAjaHDRColorimetry(EAjaHDRMetadataGamut Colorimetry);
+			/** Convert HDR colorimetry from Aja to unreal format. */
+			static EAjaHDRMetadataGamut ConvertFromAjaHDRColorimetry(NTV2HDRColorimetry);
+			/** Convert HDR luminance from Unreal to Aja format. */
+			static NTV2HDRLuminance ConvertToAjaHDRLuminance(EAjaHDRMetadataLuminance HDRLuminance);
+			/** Convert HDR luminance from Aja to Unreal format. */
+			static EAjaHDRMetadataLuminance ConvertFromAjaHDRLuminance(NTV2HDRLuminance AjaHDRLuminance);
+			
 			static NTV2FrameRate ConvertToFrameRate(uint32_t InNumerator, uint32_t InDenominator);
 			static TimecodeFormat ConvertToTimecodeFormat(NTV2VideoFormat InVideoFormat);
 
 			static bool TryVideoFormatIndexToNTV2VideoFormat(FAJAVideoFormat InVideoFormatIndex, NTV2VideoFormat& OutFoundVideoFormat);
 			static bool GetInputVideoFormat(CNTV2Card* InCard, ETransportType InTransportType, NTV2Channel InChannel, NTV2InputSource InInputSource, NTV2VideoFormat InExpectedVideoFormat, NTV2VideoFormat& OutFoundVideoFormat, bool bSetSDIConversion, std::string& OutFailedReason, bool bEnforceExpectedFormat = true);
+			/** Fetch HDR Metadata from VPID registers. */
+			static bool GetInputHDRMetadata(CNTV2Card* InCard, NTV2Channel InChannel, FAjaHDROptions& OutHDRMetadata);
 			static std::optional<NTV2VideoFormat> GetInputVideoFormat(CNTV2Card* InCard, ETransportType InTransportType, NTV2Channel InChannel, NTV2InputSource InInputSource, NTV2VideoFormat InExpectedVideoFormat, bool bSetSDIConversion, std::string& OutFailedReason);
 			static bool CompareFormats(NTV2VideoFormat LHS, NTV2VideoFormat& RHS, std::string& OutFailedReason);
 			static void RouteSdiSignal(CNTV2Card* InCard, ETransportType InTransportType, NTV2Channel InChannel, NTV2VideoFormat InVideoFormat, NTV2FrameBufferFormat InPixelFormat, bool bIsInput, bool bIsInputColorRgb, bool bWillUseKey);
diff --git a/Engine/Plugins/Media/AjaMedia/Source/Aja/Public/AJALib.h b/Engine/Plugins/Media/AjaMedia/Source/Aja/Public/AJALib.h
--- a/Engine/Plugins/Media/AjaMedia/Source/Aja/Public/AJALib.h
+++ b/Engine/Plugins/Media/AjaMedia/Source/Aja/Public/AJALib.h
@@ -38,7 +38,7 @@ namespace AJA
 		PF_10BIT_RGB,	// As Input/Output
 		PF_10BIT_YCBCR,	// As Input/Output
 	};
-
+	
 	/* SDI transport type
 	*****************************************************************************/
 	enum class ETransportType
@@ -364,6 +364,52 @@ namespace AJA
 		uint8_t* VideoBuffer;
 	};
 
+	
+	/**
+	 * HDR Transfer function.
+	 */
+	enum class EAjaHDRMetadataEOTF : uint8
+	{
+		SDR,
+		HLG,
+		PQ,
+		Unspecified
+	};
+	
+	/**
+	 * HDR Color Gamut.
+	 */
+	enum class EAjaHDRMetadataGamut : uint8
+	{
+		Rec709,
+		Rec2020,
+		Invalid
+	};
+
+	/**
+	 * HDR Luminance.
+	 */
+	enum class EAjaHDRMetadataLuminance : uint8
+	{
+		YCbCr,
+		ICtCp
+	};
+
+	/**
+	 * Set of metadata describing a HDR video signal.
+	 */
+	struct AJACORE_API FAjaHDROptions
+	{
+		/** Transfer function to use for converting the video signal to an optical signal. */
+		EAjaHDRMetadataEOTF EOTF = EAjaHDRMetadataEOTF::SDR;
+
+		/** The color gamut of the video signal. */
+		EAjaHDRMetadataGamut Gamut = EAjaHDRMetadataGamut::Rec709;
+
+		/** Color representation format of the video signal. */
+		EAjaHDRMetadataLuminance Luminance = EAjaHDRMetadataLuminance::YCbCr;
+	};
+
 	struct AJACORE_API IAJAInputOutputChannelCallbackInterface : IAJASyncChannelCallbackInterface
 	{
 		IAJAInputOutputChannelCallbackInterface();
@@ -399,6 +445,7 @@ namespace AJA
 		uint32_t OutputNumberOfBuffers; // [1...x] supported but not suggested (min of 2 is suggested)
 		FAJAVideoFormat VideoFormatIndex;
 		EPixelFormat PixelFormat;
+		FAjaHDROptions HDROptions;
 		ETimecodeFormat TimecodeFormat;
 		EAJAReferenceType OutputReferenceType;
 		uint32_t BurnTimecodePercentY;
diff --git a/Engine/Plugins/Media/AjaMedia/Source/AjaMediaOutput/Private/AjaMediaCapture.cpp b/Engine/Plugins/Media/AjaMedia/Source/AjaMediaOutput/Private/AjaMediaCapture.cpp
--- a/Engine/Plugins/Media/AjaMedia/Source/AjaMediaOutput/Private/AjaMediaCapture.cpp
+++ b/Engine/Plugins/Media/AjaMedia/Source/AjaMediaOutput/Private/AjaMediaCapture.cpp
@@ -7,12 +7,12 @@
 #include "AjaMediaOutput.h"
 #include "AjaMediaOutputModule.h"
 #include "GPUTextureTransferModule.h"
-#include "Misc/CoreDelegates.h"
 #include "Engine/Engine.h"
 #include "HAL/Event.h"
 #include "HAL/IConsoleManager.h"
 #include "IAjaMediaOutputModule.h"
 #include "IAjaMediaModule.h"
+#include "MediaIOCoreDefinitions.h"
 #include "MediaIOCoreFileWriter.h"
 #include "MediaIOCoreSubsystem.h"
 #include "Misc/ScopeLock.h"
@@ -97,6 +97,16 @@ namespace AjaMediaCaptureDevice
 
 namespace AjaMediaCaptureUtils
 {
+	AJA::FAjaHDROptions MakeAjaHDRMetadata(const FAjaMediaHDROptions& HDROptions)
+	{
+		AJA::FAjaHDROptions HDRMetadata;
+		
+		HDRMetadata.Gamut = (AJA::EAjaHDRMetadataGamut) HDROptions.Gamut;
+		HDRMetadata.EOTF = (AJA::EAjaHDRMetadataEOTF) HDROptions.EOTF;
+
+		return HDRMetadata;
+	}
+	
 	AJA::ETransportType ConvertTransportType(const EMediaIOTransportType TransportType, const EMediaIOQuadLinkTransportType QuadTransportType)
 	{
 		switch (TransportType)
@@ -425,6 +435,25 @@ bool UAjaMediaCapture::HasFinishedProcessing() const
 	return Super::HasFinishedProcessing() || !OutputChannel;
 }
 
+const FMatrix& UAjaMediaCapture::GetRGBToYUVConversionMatrix() const
+{
+	if (const UAjaMediaOutput* AjaOutput = Cast<UAjaMediaOutput>(MediaOutput))
+	{
+		switch(AjaOutput->HDROptions.Gamut)
+		{
+		case EAjaHDRMetadataGamut::Rec709:
+			return MediaShaders::RgbToYuvRec709Scaled;
+		case EAjaHDRMetadataGamut::Rec2020:
+			return MediaShaders::RgbToYuvRec2020Scaled;
+		default:
+			checkNoEntry();
+			return MediaShaders::RgbToYuvRec709Scaled;
+		}
+	}
+	
+	return Super::GetRGBToYUVConversionMatrix();
+}
+
 bool UAjaMediaCapture::InitAJA(UAjaMediaOutput* InAjaMediaOutput)
 {
 	check(InAjaMediaOutput);
@@ -484,6 +513,8 @@ bool UAjaMediaCapture::InitAJA(UAjaMediaOutput* InAjaMediaOutput)
 	ChannelOptions.bUseGPUDMA = ShouldCaptureRHIResource();
 	ChannelOptions.bDirectlyWriteAudio = InAjaMediaOutput->bOutputAudioOnAudioThread;
 
+	ChannelOptions.HDROptions = AjaMediaCaptureUtils::MakeAjaHDRMetadata(InAjaMediaOutput->HDROptions);
+	
 	bOutputAudio = InAjaMediaOutput->bOutputAudio;
 	bDirectlyWriteAudio = InAjaMediaOutput->bOutputAudioOnAudioThread;
 	
diff --git a/Engine/Plugins/Media/AjaMedia/Source/AjaMediaOutput/Public/AjaMediaCapture.h b/Engine/Plugins/Media/AjaMedia/Source/AjaMediaOutput/Public/AjaMediaCapture.h
--- a/Engine/Plugins/Media/AjaMedia/Source/AjaMediaOutput/Public/AjaMediaCapture.h
+++ b/Engine/Plugins/Media/AjaMedia/Source/AjaMediaOutput/Public/AjaMediaCapture.h
@@ -60,8 +60,7 @@ protected:
 	{
 		return true;
 	}
-
-
+	virtual const FMatrix& GetRGBToYUVConversionMatrix() const override;
 
 private:
 	struct FAjaOutputCallback;
diff --git a/Engine/Plugins/Media/AjaMedia/Source/AjaMediaOutput/Public/AjaMediaOutput.h b/Engine/Plugins/Media/AjaMedia/Source/AjaMediaOutput/Public/AjaMediaOutput.h
--- a/Engine/Plugins/Media/AjaMedia/Source/AjaMediaOutput/Public/AjaMediaOutput.h
+++ b/Engine/Plugins/Media/AjaMedia/Source/AjaMediaOutput/Public/AjaMediaOutput.h
@@ -3,7 +3,6 @@
 #pragma once
 
 #include "MediaOutput.h"
-
 #include "MediaIOCoreDefinitions.h"
 
 #include "AjaMediaOutput.generated.h"
@@ -35,6 +34,46 @@ enum class EAjaMediaOutputChannelConfiguration : uint8
 	CH_16 = 16 UMETA(DisplayName = "16 Channels")
 };
 
+/**
+ * HDR Transfer function.
+ * Must match AJA::EAjaHDRMetadataEOTF.
+ */
+UENUM()
+enum class EAjaHDRMetadataEOTF : uint8
+{
+	SDR,
+	HLG,
+	PQ
+};
+	
+/**
+ * HDR Color Gamut.
+ * Must match AJA::EAjaHDRMetadataGamut.
+ */
+UENUM()
+enum class EAjaHDRMetadataGamut : uint8
+{
+	Rec709,
+	Rec2020
+};
+
+/**
+ * Set of metadata describing a HDR video signal.
+ */
+USTRUCT()
+struct AJAMEDIAOUTPUT_API FAjaMediaHDROptions
+{
+	GENERATED_BODY()
+
+	/** Transfer function to use for converting the video signal to an optical signal. */
+	UPROPERTY(EditAnywhere, Category = "HDR")
+	EAjaHDRMetadataEOTF EOTF = EAjaHDRMetadataEOTF::SDR;
+
+	/** The color gamut of the video signal. */
+	UPROPERTY(EditAnywhere, Category = "HDR")
+	EAjaHDRMetadataGamut Gamut = EAjaHDRMetadataGamut::Rec709;
+};
+
 /**
  * Output information for an aja media capture.
  * @note	'Frame Buffer Pixel Format' must be set to at least 8 bits of alpha to enabled the Key.
@@ -111,6 +150,12 @@ public:
 	UPROPERTY(EditAnywhere, AdvancedDisplay, Category="Output")
 	bool bInterlacedFieldsTimecodeNeedToMatch;
 
+	/**
+	 * HDR Metadata of the video signal.
+	 */
+	UPROPERTY(EditAnywhere, Category = "Output", DisplayName="HDR")
+	FAjaMediaHDROptions HDROptions;
+
 	/** Try to maintain a the engine "Genlock" with the VSync signal. */
 	UPROPERTY(BlueprintReadWrite, EditAnywhere, Category="Synchronization")
 	bool bWaitForSyncEvent;
