Replace LogTemp references in AjaCore module

diff --git a/Engine/Plugins/Media/AjaMedia/Source/Aja/Private/AjaCoreModule.cpp b/Engine/Plugins/Media/AjaMedia/Source/Aja/Private/AjaCoreModule.cpp
--- a/Engine/Plugins/Media/AjaMedia/Source/Aja/Private/AjaCoreModule.cpp
+++ b/Engine/Plugins/Media/AjaMedia/Source/Aja/Private/AjaCoreModule.cpp
@@ -1,5 +1,6 @@
 // Copyright Epic Games, Inc. All Rights Reserved.
 #include "AjaCoreModule.h"
 
+DEFINE_LOG_CATEGORY(LogAjaCore);
 
 IMPLEMENT_MODULE(FAjaCoreModule, AjaCore);
diff --git a/Engine/Plugins/Media/AjaMedia/Source/Aja/Private/ChannelBase.cpp b/Engine/Plugins/Media/AjaMedia/Source/Aja/Private/ChannelBase.cpp
--- a/Engine/Plugins/Media/AjaMedia/Source/Aja/Private/ChannelBase.cpp
+++ b/Engine/Plugins/Media/AjaMedia/Source/Aja/Private/ChannelBase.cpp
@@ -71,7 +71,7 @@ namespace AJA
 		{
 			if (Options.ChannelIndex > NTV2_MAX_NUM_CHANNELS || Options.ChannelIndex < 1)
 			{
-				UE_LOG(LogTemp, Error, TEXT("AJAChannel: The port index '%d' is invalid.\n"), Options.ChannelIndex);
+				UE_LOG(LogAjaCore, Error, TEXT("AJAChannel: The port index '%d' is invalid.\n"), Options.ChannelIndex);
 			}
 			else
 			{
@@ -83,12 +83,12 @@ namespace AJA
 			{
 				if (Options.KeyChannelIndex > NTV2_MAX_NUM_CHANNELS || Options.KeyChannelIndex < 1)
 				{
-					UE_LOG(LogTemp, Error, TEXT("AJAChannel: The key port index '%d' is invalid.\n"), Options.KeyChannelIndex);
+					UE_LOG(LogAjaCore, Error, TEXT("AJAChannel: The key port index '%d' is invalid.\n"), Options.KeyChannelIndex);
 					Options.bUseKey = false;
 				}
 				else if (Options.KeyChannelIndex == Options.ChannelIndex)
 				{
-					UE_LOG(LogTemp, Error, TEXT("AJAChannel: The key port index '%d' is the same as the port index '%d'.\n"), Options.KeyChannelIndex, Options.ChannelIndex);
+					UE_LOG(LogAjaCore, Error, TEXT("AJAChannel: The key port index '%d' is the same as the port index '%d'.\n"), Options.KeyChannelIndex, Options.ChannelIndex);
 					Options.bUseKey = false;
 				}
 				else
@@ -123,7 +123,7 @@ namespace AJA
 			{
 				if (Channel != Helpers::GetTransportTypeChannel(GetOptions().TransportType, Channel))
 				{
-					UE_LOG(LogTemp, Error, TEXT("Device: Can't do a %S on channel '%d'. Did you mean channel '%d'?")
+					UE_LOG(LogAjaCore, Error, TEXT("Device: Can't do a %S on channel '%d'. Did you mean channel '%d'?")
 						, Helpers::TransportTypeToString(GetOptions().TransportType)
 						, uint32_t(Channel) + 1
 						, uint32_t(Helpers::GetTransportTypeChannel(GetOptions().TransportType, Channel)) + 1);
@@ -133,7 +133,7 @@ namespace AJA
 				{
 					if (KeyChannel != Helpers::GetTransportTypeChannel(GetOptions().TransportType, KeyChannel))
 					{
-						UE_LOG(LogTemp, Error, TEXT("Device: Can't do a %S on channel '%d'. Did you mean channel '%d'?")
+						UE_LOG(LogAjaCore, Error, TEXT("Device: Can't do a %S on channel '%d'. Did you mean channel '%d'?")
 							, Helpers::TransportTypeToString(GetOptions().TransportType)
 							, uint32_t(KeyChannel) + 1
 							, uint32_t(Helpers::GetTransportTypeChannel(GetOptions().TransportType, KeyChannel)) + 1);
@@ -145,7 +145,7 @@ namespace AJA
 					NTV2Channel SyncChannel = (NTV2Channel)(Options.SynchronizeChannelIndex - 1);
 					if (SyncChannel != Helpers::GetTransportTypeChannel(GetOptions().TransportType, SyncChannel))
 					{
-						UE_LOG(LogTemp, Error, TEXT("Device: Can't do a %S on channel '%d'. Did you mean channel '%d'?")
+						UE_LOG(LogAjaCore, Error, TEXT("Device: Can't do a %S on channel '%d'. Did you mean channel '%d'?")
 							, Helpers::TransportTypeToString(GetOptions().TransportType)
 							, uint32_t(SyncChannel) + 1
 							, uint32_t(Helpers::GetTransportTypeChannel(GetOptions().TransportType, SyncChannel)) + 1);
@@ -276,7 +276,7 @@ namespace AJA
 
 			if (!Helpers::TryVideoFormatIndexToNTV2VideoFormat(GetOptions().VideoFormatIndex, DesiredVideoFormat))
 			{
-				UE_LOG(LogTemp, Error, TEXT("ConfigureVideo: The expected video format is invalid for %d.\n"), uint32_t(Channel) + 1);
+				UE_LOG(LogAjaCore, Error, TEXT("ConfigureVideo: The expected video format is invalid for %d.\n"), uint32_t(Channel) + 1);
 				return false;
 			}
 
@@ -287,13 +287,13 @@ namespace AJA
 
 			if (Options.TransportType == ETransportType::TT_SdiSingle4kTSI && GetOptions().bUseKey)
 			{
-				UE_LOG(LogTemp, Error, TEXT("ConfigureVideo: 4k Key and Fill output is not supported for this card.\n"));
+				UE_LOG(LogAjaCore, Error, TEXT("ConfigureVideo: 4k Key and Fill output is not supported for this card.\n"));
 				return false;
 			}
 
 			if (Options.bDirectlyWriteAudio && Options.bUseAutoCirculating)
 			{
-				UE_LOG(LogTemp, Error, TEXT("ConfigureVideo: Direct audio write is only available outside of auto-circulate mode..\n"));
+				UE_LOG(LogAjaCore, Error, TEXT("ConfigureVideo: Direct audio write is only available outside of auto-circulate mode..\n"));
 				return false;
 			}
 			
@@ -407,7 +407,7 @@ namespace AJA
 					std::string FailureReason;
 					if (!Helpers::GetInputVideoFormat(Device->GetCard(), GetOptions().TransportType, Channel, InputSource, DesiredVideoFormat, VideoFormat, false, FailureReason))
 					{
-						UE_LOG(LogTemp, Error, TEXT("Device: Initialization of the input failed for channel %d on device '%S'. '%S'"), uint32_t(Channel) + 1, GetDevice().GetDisplayName().c_str(), FailureReason.c_str());
+						UE_LOG(LogAjaCore, Error, TEXT("Device: Initialization of the input failed for channel %d on device '%S'. '%S'"), uint32_t(Channel) + 1, GetDevice().GetDisplayName().c_str(), FailureReason.c_str());
 						return false;
 					}
 				}
@@ -418,14 +418,14 @@ namespace AJA
 					VideoFormat = DesiredVideoFormat;
 					if (VideoFormat == NTV2_FORMAT_UNKNOWN)
 					{
-						UE_LOG(LogTemp, Error, TEXT("ConfigureVideo: Unknown video format for output channel %d on device '%S'.\n"), uint32_t(Channel) + 1, GetDevice().GetDisplayName().c_str());
+						UE_LOG(LogAjaCore, Error, TEXT("ConfigureVideo: Unknown video format for output channel %d on device '%S'.\n"), uint32_t(Channel) + 1, GetDevice().GetDisplayName().c_str());
 						return false;
 					}
 				}
 
 				if (!::NTV2DeviceCanDoVideoFormat(Device->GetCard()->GetDeviceID(), VideoFormat))
 				{
-					UE_LOG(LogTemp, Error, TEXT("ConfigureVideo: The device '%S' doesn't support the video format '%S'.\n"), GetDevice().GetDisplayName().c_str(), ::NTV2VideoFormatToString(VideoFormat).c_str());
+					UE_LOG(LogAjaCore, Error, TEXT("ConfigureVideo: The device '%S' doesn't support the video format '%S'.\n"), GetDevice().GetDisplayName().c_str(), ::NTV2VideoFormatToString(VideoFormat).c_str());
 					return false;
 				}
 
@@ -434,7 +434,7 @@ namespace AJA
 
 				if (!::NTV2DeviceCanDoFrameBufferFormat(Device->GetCard()->GetDeviceID(), PixelFormat))
 				{
-					UE_LOG(LogTemp, Error, TEXT("ConfigureVideo: The device '%S' doesn't support the pixel format '%S'.\n"), GetDevice().GetDisplayName().c_str(), ::NTV2FrameBufferFormatToString(PixelFormat).c_str());
+					UE_LOG(LogAjaCore, Error, TEXT("ConfigureVideo: The device '%S' doesn't support the pixel format '%S'.\n"), GetDevice().GetDisplayName().c_str(), ::NTV2FrameBufferFormatToString(PixelFormat).c_str());
 					return false;
 				}
 
@@ -457,7 +457,7 @@ namespace AJA
 					{
 						if (GetOptions().SynchronizeChannelIndex > NTV2_MAX_NUM_CHANNELS || GetOptions().SynchronizeChannelIndex < 1)
 						{
-							UE_LOG(LogTemp, Error, TEXT("ConfigureVideo: The synchronization port index '%d' is invalid.\n"), GetOptions().SynchronizeChannelIndex);
+							UE_LOG(LogAjaCore, Error, TEXT("ConfigureVideo: The synchronization port index '%d' is invalid.\n"), GetOptions().SynchronizeChannelIndex);
 							return false;
 						}
 						else
@@ -489,7 +489,7 @@ namespace AJA
 			}
 			else
 			{
-				UE_LOG(LogTemp, Error, TEXT("ConfigureVideo: Invalid video options for channel %d on device '%S'."), uint32_t(Channel) + 1, GetDevice().GetDisplayName().c_str());
+				UE_LOG(LogAjaCore, Error, TEXT("ConfigureVideo: Invalid video options for channel %d on device '%S'."), uint32_t(Channel) + 1, GetDevice().GetDisplayName().c_str());
 				return false;
 			}
 
@@ -528,7 +528,7 @@ namespace AJA
 					{
 						if (GetOptions().bConvertOutputLevelAToB)
 						{
-							UE_LOG(LogTemp, Warning,  TEXT("ConfigureVideo: Can't convert level A to B for channel %d on device '%S'."), uint32_t(Channel) + 1, GetDevice().GetDisplayName().c_str());
+							UE_LOG(LogAjaCore, Warning,  TEXT("ConfigureVideo: Can't convert level A to B for channel %d on device '%S'."), uint32_t(Channel) + 1, GetDevice().GetDisplayName().c_str());
 						}
 					}
 
@@ -560,7 +560,7 @@ namespace AJA
 							{
 								if (GetOptions().bConvertOutputLevelAToB)
 								{
-									UE_LOG(LogTemp, Warning,  TEXT("ConfigureVideo: Can't convert level A to B on key for channel %d on device '%S'."), uint32_t(Channel) + 1, GetDevice().GetDisplayName().c_str());
+									UE_LOG(LogAjaCore, Warning,  TEXT("ConfigureVideo: Can't convert level A to B on key for channel %d on device '%S'."), uint32_t(Channel) + 1, GetDevice().GetDisplayName().c_str());
 								}
 							}
 							Helpers::RouteKeySignal(Device->GetCard(), GetOptions().TransportType, Channel, KeyChannel, InPixelFormat, false);
@@ -620,7 +620,7 @@ namespace AJA
 				bool bSupportAnc = ::NTV2DeviceCanDoCustomAnc(GetDevice().GetDeviceID());
 				if (!bSupportAnc)
 				{
-					UE_LOG(LogTemp, Warning,  TEXT("ConfigureAncillary: Device doesn't support Anc.\n"));
+					UE_LOG(LogAjaCore, Warning,  TEXT("ConfigureAncillary: Device doesn't support Anc.\n"));
 					Options.bUseAncillary = false;
 				}
 			}
@@ -647,7 +647,7 @@ namespace AJA
 				AudioSystem = NTV2InputSourceToAudioSystem(InputSource);
 				if (NTV2_IS_VALID_AUDIO_SYSTEM(AudioSystem) == false)
 				{
-					UE_LOG(LogTemp, Error, TEXT("ConfigureAudio: Could not get valid AudioSystem from InputSource %d for channel %d on device %S."), uint32_t(InputSource), uint32_t(Channel) + 1, GetDevice().GetDisplayName().c_str());
+					UE_LOG(LogAjaCore, Error, TEXT("ConfigureAudio: Could not get valid AudioSystem from InputSource %d for channel %d on device %S."), uint32_t(InputSource), uint32_t(Channel) + 1, GetDevice().GetDisplayName().c_str());
 					return false;
 				}
 			}
@@ -655,7 +655,7 @@ namespace AJA
 			{
 				if (UseAudio())
 				{
-					UE_LOG(LogTemp, Error, TEXT("ConfigureAudio: Invalid audio options for channel %d on device %S."), uint32_t(Channel) + 1, GetDevice().GetDisplayName().c_str());
+					UE_LOG(LogAjaCore, Error, TEXT("ConfigureAudio: Invalid audio options for channel %d on device %S."), uint32_t(Channel) + 1, GetDevice().GetDisplayName().c_str());
 					return false;
 				}
 			}
@@ -683,7 +683,7 @@ namespace AJA
 				if (NumberOfAudioChannel > MaxNumOfChannels)
 				{
 					NumberOfAudioChannel = MaxNumOfChannels;
-					UE_LOG(LogTemp, Warning,  TEXT("ConfigureAudio: Changed number of audio channel to %d.\n"), NumberOfAudioChannel);
+					UE_LOG(LogAjaCore, Warning,  TEXT("ConfigureAudio: Changed number of audio channel to %d.\n"), NumberOfAudioChannel);
 				}
 
 				Card.SetNumberAudioChannels(NumberOfAudioChannel, AudioSystem);
diff --git a/Engine/Plugins/Media/AjaMedia/Source/Aja/Private/Device.cpp b/Engine/Plugins/Media/AjaMedia/Source/Aja/Private/Device.cpp
--- a/Engine/Plugins/Media/AjaMedia/Source/Aja/Private/Device.cpp
+++ b/Engine/Plugins/Media/AjaMedia/Source/Aja/Private/Device.cpp
@@ -19,7 +19,7 @@ namespace AJA
 		{
 			if (InDeviceOption.DeviceIndex >= MaxNumberOfDevice || InDeviceOption.DeviceIndex < 0)
 			{
-				UE_LOG(LogTemp, Error, TEXT("DeviceCache: The device index '%d' is invalid.\n"), InDeviceOption.DeviceIndex);
+				UE_LOG(LogAjaCore, Error, TEXT("DeviceCache: The device index '%d' is invalid.\n"), InDeviceOption.DeviceIndex);
 				return std::shared_ptr<DeviceConnection>();
 			}
 
@@ -240,7 +240,7 @@ namespace AJA
 				{
 					if (!FoundChannel->bIsAutoDetected)
 					{
-						UE_LOG(LogTemp, Warning,  TEXT("Device: Autodetect reinitialized channel '%d' on device '%S' but the channel was not in autodetect mode.\n")
+						UE_LOG(LogAjaCore, Warning,  TEXT("Device: Autodetect reinitialized channel '%d' on device '%S' but the channel was not in autodetect mode.\n")
 							, uint32_t(InChannel) + 1, Connection.Card->GetDisplayName().c_str());
 					}
 
@@ -254,7 +254,7 @@ namespace AJA
 
 				if (bResult && FoundChannel->bIsOwned && bInAsOwner)
 				{
-					UE_LOG(LogTemp, Error, TEXT("Device: Trying to enable the channel '%d' on device '%S' as the owner but it was already enabled. Do you already have an '%S' running on the same channel?\n")
+					UE_LOG(LogAjaCore, Error, TEXT("Device: Trying to enable the channel '%d' on device '%S' as the owner but it was already enabled. Do you already have an '%S' running on the same channel?\n")
 						, uint32_t(InChannel) + 1, Connection.Card->GetDisplayName().c_str()
 						, FoundChannel->bIsInput ? "input" : "output");
 					bResult = false;
@@ -262,7 +262,7 @@ namespace AJA
 
 				if (bResult && FoundChannel->TransportType != InTransportType)
 				{
-					UE_LOG(LogTemp, Error, TEXT("Device: Trying to enable the channel '%d' on device '%S' as a '%S' but it was already enabled as a '%S'.\n")
+					UE_LOG(LogAjaCore, Error, TEXT("Device: Trying to enable the channel '%d' on device '%S' as a '%S' but it was already enabled as a '%S'.\n")
 						, uint32_t(InChannel) + 1
 						, Connection.Card->GetDisplayName().c_str()
 						, Helpers::TransportTypeToString(InTransportType)
@@ -272,7 +272,7 @@ namespace AJA
 
 				if (bResult && FoundChannel->bIsInput != bInIsInput)
 				{
-					UE_LOG(LogTemp, Error, TEXT("Device: Trying to enable the channel '%d' on device '%S' as an '%S' but it was already enabled as an '%S'.\n")
+					UE_LOG(LogAjaCore, Error, TEXT("Device: Trying to enable the channel '%d' on device '%S' as an '%S' but it was already enabled as an '%S'.\n")
 						, uint32_t(InChannel) + 1
 						, Connection.Card->GetDisplayName().c_str()
 						, bInIsInput ? "input" : "output"
@@ -285,14 +285,14 @@ namespace AJA
 					std::string FailureReason;
 					if (!Helpers::GetInputVideoFormat(Connection.Card, InTransportType, InChannel, InInputSource, InDesiredInputFormat, VideoFormat, false, FailureReason))
 					{
-						UE_LOG(LogTemp, Error, TEXT("Device: Initialization of the input failed for channel %d on device %S. %S"), uint32_t(InChannel) + 1, Connection.Card->GetDisplayName().c_str(), FailureReason.c_str());
+						UE_LOG(LogAjaCore, Error, TEXT("Device: Initialization of the input failed for channel %d on device %S. %S"), uint32_t(InChannel) + 1, Connection.Card->GetDisplayName().c_str(), FailureReason.c_str());
 						bResult = false;
 					}
 
 					if (bResult && !bInIsAutoDetected && VideoFormat != FoundChannel->VideoFormat)
 					{
 						const bool bForRetailDisplay = true;
-						UE_LOG(LogTemp, Error, TEXT("Device: Try to enable with the video format '%S' the channel '%d' on device '%S' that was already enabled with '%S'.\n")
+						UE_LOG(LogAjaCore, Error, TEXT("Device: Try to enable with the video format '%S' the channel '%d' on device '%S' that was already enabled with '%S'.\n")
 							, NTV2FrameRateToString(GetNTV2FrameRateFromVideoFormat(VideoFormat), bForRetailDisplay).c_str()
 							, uint32_t(InChannel) + 1, Connection.Card->GetDisplayName().c_str()
 							, NTV2FrameRateToString(GetNTV2FrameRateFromVideoFormat(FoundChannel->VideoFormat), bForRetailDisplay).c_str());
@@ -302,7 +302,7 @@ namespace AJA
 
 				if (bResult && FoundChannel->TimecodeFormat != InTimecodeFormat && !bInIsAutoDetected)
 				{
-					UE_LOG(LogTemp, Warning,  TEXT("Device: Try to enable the channel '%d' on device '%S' with a timecode format that is not the same as the previous one. Timecode may not be decoded properly.\n")
+					UE_LOG(LogAjaCore, Warning,  TEXT("Device: Try to enable the channel '%d' on device '%S' with a timecode format that is not the same as the previous one. Timecode may not be decoded properly.\n")
 						, uint32_t(InChannel) + 1, Connection.Card->GetDisplayName().c_str());
 				}
 
@@ -357,7 +357,7 @@ namespace AJA
 			auto FoundIterator = std::find_if(std::begin(ConnectionChannelInfos), std::end(ConnectionChannelInfos), [=](const ChannelInfo* Other) { return Other->Channel == InChannel; });
 			if (FoundIterator == std::end(ConnectionChannelInfos))
 			{
-				UE_LOG(LogTemp, Error, TEXT("Device: Try to unregister the channel '%d' on device '%S' when it's not registered.\n"), uint32_t(InChannel) + 1, Connection.Card->GetDisplayName().c_str());
+				UE_LOG(LogAjaCore, Error, TEXT("Device: Try to unregister the channel '%d' on device '%S' when it's not registered.\n"), uint32_t(InChannel) + 1, Connection.Card->GetDisplayName().c_str());
 				assert(false);
 				return false;
 			}
@@ -462,7 +462,7 @@ namespace AJA
 			
 			if (!NTV2_IS_VALID_NTV2FrameRate(InFrameRate))
 			{
-				UE_LOG(LogTemp, Error, TEXT("Device: Trying to read LTC from the reference input '%d' on device '%S' but it is not supported by the device.\n"), uint32_t(InSource) + 1, Connection.Card->GetDisplayName().c_str());
+				UE_LOG(LogAjaCore, Error, TEXT("Device: Trying to read LTC from the reference input '%d' on device '%S' but it is not supported by the device.\n"), uint32_t(InSource) + 1, Connection.Card->GetDisplayName().c_str());
 				return false;
 			}
 
@@ -470,14 +470,14 @@ namespace AJA
 			bool bResult = true;
 			if (bResult && ::NTV2DeviceGetNumLTCInputs(Connection.Card->GetDeviceID()) <= (uint32_t)InSource)
 			{
-				UE_LOG(LogTemp, Error, TEXT("Device: Trying to read LTC from the LTC input '%d' on device '%S' but it is not supported by the device.\n"), uint32_t(InSource) + 1, Connection.Card->GetDisplayName().c_str());
+				UE_LOG(LogAjaCore, Error, TEXT("Device: Trying to read LTC from the LTC input '%d' on device '%S' but it is not supported by the device.\n"), uint32_t(InSource) + 1, Connection.Card->GetDisplayName().c_str());
 				bResult = false;
 			}
 			
 			//If there are other LTC readers, exit, only one LTC input is supported
 			if (Connection.LTCFrameRate != NTV2_FRAMERATE_INVALID)
 			{
-				UE_LOG(LogTemp, Error, TEXT("Device: The device '%S' is already reading timecode from LTC input with frame rate '%S' and only one is supported."), Connection.Card->GetDisplayName().c_str(), NTV2FrameRateToString(Connection.LTCFrameRate).c_str());
+				UE_LOG(LogAjaCore, Error, TEXT("Device: The device '%S' is already reading timecode from LTC input with frame rate '%S' and only one is supported."), Connection.Card->GetDisplayName().c_str(), NTV2FrameRateToString(Connection.LTCFrameRate).c_str());
 				bResult = false;
 			}
 
@@ -486,21 +486,21 @@ namespace AJA
 			{
 				if (!::NTV2DeviceCanDoLTCInOnRefPort(Connection.Card->GetDeviceID()))
 				{
-					UE_LOG(LogTemp, Error, TEXT("Device: Trying to read LTC from reference pin on device '%S' but it doesn't have that capability.\n"), Connection.Card->GetDisplayName().c_str());
+					UE_LOG(LogAjaCore, Error, TEXT("Device: Trying to read LTC from reference pin on device '%S' but it doesn't have that capability.\n"), Connection.Card->GetDisplayName().c_str());
 					bResult = false;
 				}
 
 				//Verify if we are already reading LTC from ref pin
 				if (bResult && Connection.bAnalogLtcFromReferenceInput)
 				{
-					UE_LOG(LogTemp, Error, TEXT("Device: The reference pin is already used to read LTC timecode on device '%S'.\n"), Connection.Card->GetDisplayName().c_str());
+					UE_LOG(LogAjaCore, Error, TEXT("Device: The reference pin is already used to read LTC timecode on device '%S'.\n"), Connection.Card->GetDisplayName().c_str());
 					bResult = false;
 				}
 
 				//Verify if we are already using ref pin for genlock
 				if (bResult && Connection.bOutputReferenceSet && Connection.OutputReferenceType == EAJAReferenceType::EAJA_REFERENCETYPE_EXTERNAL)
 				{
-					UE_LOG(LogTemp, Error, TEXT("Device: Trying to read LTC from the reference input on device '%S' but the reference is already used as a genlock source.\n"), Connection.Card->GetDisplayName().c_str());
+					UE_LOG(LogAjaCore, Error, TEXT("Device: Trying to read LTC from the reference input on device '%S' but the reference is already used as a genlock source.\n"), Connection.Card->GetDisplayName().c_str());
 					bResult = false;
 				}
 			}
@@ -518,13 +518,13 @@ namespace AJA
 					{
 						if (!::IsMultiFormatCompatible(InFrameRate, ChannelFrameRate))
 						{
-							UE_LOG(LogTemp, Error, TEXT("Device: Trying to read LTC with FrameRate '%S' but it's not compatible with card's (channel 1) current FrameRate of '%S' on device '%S'.\n"), NTV2FrameRateToString(InFrameRate).c_str(), NTV2FrameRateToString(ChannelFrameRate).c_str(), Connection.Card->GetDisplayName().c_str());
+							UE_LOG(LogAjaCore, Error, TEXT("Device: Trying to read LTC with FrameRate '%S' but it's not compatible with card's (channel 1) current FrameRate of '%S' on device '%S'.\n"), NTV2FrameRateToString(InFrameRate).c_str(), NTV2FrameRateToString(ChannelFrameRate).c_str(), Connection.Card->GetDisplayName().c_str());
 							bResult = false;
 						}
 					}
 					else if (InFrameRate != ChannelFrameRate)
 					{
-						UE_LOG(LogTemp, Error, TEXT("Device: Trying to read LTC with FrameRate '%S' but card's (channel 1) current FrameRate is '%S' and device '%S' can't do multi-format.\n"), NTV2FrameRateToString(InFrameRate).c_str(), NTV2FrameRateToString(ChannelFrameRate).c_str(), Connection.Card->GetDisplayName().c_str());
+						UE_LOG(LogAjaCore, Error, TEXT("Device: Trying to read LTC with FrameRate '%S' but card's (channel 1) current FrameRate is '%S' and device '%S' can't do multi-format.\n"), NTV2FrameRateToString(InFrameRate).c_str(), NTV2FrameRateToString(ChannelFrameRate).c_str(), Connection.Card->GetDisplayName().c_str());
 						bResult = false;
 					}
 				}
@@ -586,13 +586,13 @@ namespace AJA
 			{
 				if (Connection.OutputReferenceType != InOutputReferenceType)
 				{
-					UE_LOG(LogTemp, Error, TEXT("Device: Couldn't set the reference for output on device %S. The type was already setup with %d.\n"), Connection.Card->GetDisplayName().c_str(), int(Connection.OutputReferenceType));
+					UE_LOG(LogAjaCore, Error, TEXT("Device: Couldn't set the reference for output on device %S. The type was already setup with %d.\n"), Connection.Card->GetDisplayName().c_str(), int(Connection.OutputReferenceType));
 					return false;
 				}
 
 				if (Connection.OutputReferenceType == EAJAReferenceType::EAJA_REFERENCETYPE_INPUT && Connection.OutputReferenceChannel != InOutputReferenceChannel)
 				{
-					UE_LOG(LogTemp, Error, TEXT("Device: Couldn't set the reference for output on device %S. The input was already setup with %d.\n"), Connection.Card->GetDisplayName().c_str(), Connection.OutputReferenceChannel);
+					UE_LOG(LogAjaCore, Error, TEXT("Device: Couldn't set the reference for output on device %S. The input was already setup with %d.\n"), Connection.Card->GetDisplayName().c_str(), Connection.OutputReferenceChannel);
 					return false;
 				}
 
@@ -605,7 +605,7 @@ namespace AJA
 			{
 				if (Connection.bAnalogLtcFromReferenceInput)
 				{
-					UE_LOG(LogTemp, Error, TEXT("Device: Couldn't set the reference for output channel %d on device %S. The reference is used to read analog LTC.\n"), InOutputReferenceChannel, Connection.Card->GetDisplayName().c_str());
+					UE_LOG(LogAjaCore, Error, TEXT("Device: Couldn't set the reference for output channel %d on device %S. The reference is used to read analog LTC.\n"), InOutputReferenceChannel, Connection.Card->GetDisplayName().c_str());
 					return false;
 				}
 				AJA_CHECK(Connection.Card->SetLTCInputEnable(false));
@@ -617,7 +617,7 @@ namespace AJA
 				NTV2ReferenceSource Input = NTV2InputSourceToReferenceSource(NTV2ChannelToInputSource(InOutputReferenceChannel));
 				if (Input == NTV2_NUM_REFERENCE_INPUTS)
 				{
-					UE_LOG(LogTemp, Error, TEXT("Device: Couldn't set the reference source on device '%S. The Channel index is invalid."), Connection.Card->GetDisplayName().c_str());
+					UE_LOG(LogAjaCore, Error, TEXT("Device: Couldn't set the reference source on device '%S. The Channel index is invalid."), Connection.Card->GetDisplayName().c_str());
 					return false;
 				}
 				AJA_CHECK(Connection.Card->SetReference(Input));
@@ -657,7 +657,7 @@ namespace AJA
 				const DWORD TimeoutMilli = 2000;
 				if (Counter * SleepMilli > TimeoutMilli)
 				{
-					UE_LOG(LogTemp, Warning,  TEXT("Device: Can't get the device initialized.\n"));
+					UE_LOG(LogAjaCore, Warning,  TEXT("Device: Can't get the device initialized.\n"));
 					delete Result;
 					return nullptr;
 				}
@@ -913,7 +913,7 @@ namespace AJA
 			{
 				if (TransportChannel != InChannel)
 				{
-					UE_LOG(LogTemp, Error, TEXT("Device: Can't do a %S on channel '%d'. Did you mean channel '%d'?")
+					UE_LOG(LogAjaCore, Error, TEXT("Device: Can't do a %S on channel '%d'. Did you mean channel '%d'?")
 						, Helpers::TransportTypeToString(InTransportType)
 						, uint32_t(InChannel) + 1
 						, uint32_t(TransportChannel) + 1);
@@ -933,7 +933,7 @@ namespace AJA
 				{
 					if (Itt->TransportType != InTransportType)
 					{
-						UE_LOG(LogTemp, Error, TEXT("Device: Trying to enable the channel '%d' on device '%S' as a '%S' but channel '%d' was already enabled.\n")
+						UE_LOG(LogAjaCore, Error, TEXT("Device: Trying to enable the channel '%d' on device '%S' as a '%S' but channel '%d' was already enabled.\n")
 							, int32_t(InChannel) + 1
 							, InCard->GetDisplayName().c_str()
 							, Helpers::TransportTypeToString(InTransportType)
@@ -950,7 +950,7 @@ namespace AJA
 					const bool bNewMaxInConflict = NewMaxChannel >= IttMinChannel && NewMaxChannel <= IttMaxChannel;
 					if (bNewMinInConflict || bNewMaxInConflict)
 					{
-						UE_LOG(LogTemp, Error, TEXT("Device: Trying to enable the channel '%d' on device '%S' as a '%S' but it was already enabled as a '%S' by channel '%d'.\n")
+						UE_LOG(LogAjaCore, Error, TEXT("Device: Trying to enable the channel '%d' on device '%S' as a '%S' but it was already enabled as a '%S' by channel '%d'.\n")
 							, uint32_t(InChannel) + 1
 							, InCard->GetDisplayName().c_str()
 							, Helpers::TransportTypeToString(InTransportType)
@@ -969,7 +969,7 @@ namespace AJA
 			NTV2DeviceID DeviceId = Card->GetDeviceID();
 			if (InChannel >= ::NTV2DeviceGetNumFrameStores(DeviceId))
 			{
-				UE_LOG(LogTemp, Error, TEXT("Device: The device '%S' doesn't support channel '%d'"), InCard->GetDisplayName().c_str(), uint32_t(InChannel) + 1);
+				UE_LOG(LogAjaCore, Error, TEXT("Device: The device '%S' doesn't support channel '%d'"), InCard->GetDisplayName().c_str(), uint32_t(InChannel) + 1);
 				return false;
 			}
 
@@ -994,7 +994,7 @@ namespace AJA
 				{
 					if (InChannel == NTV2Channel::NTV2_CHANNEL1 && !::IsMultiFormatCompatible(LTCFrameRate, ::GetNTV2FrameRateFromVideoFormat(OutFoundVideoFormat)))
 					{
-						UE_LOG(LogTemp, Error, TEXT("Device: The device '%S' does support Multi Format but channel %d requested format's frame rate ('%S') and it's not compatible with the LTC reference frame rate ('%S')."), InCard->GetDisplayName().c_str(), uint32_t(InChannel) + 1, NTV2FrameRateToString(::GetNTV2FrameRateFromVideoFormat(OutFoundVideoFormat)).c_str(), NTV2FrameRateToString(LTCFrameRate).c_str());
+						UE_LOG(LogAjaCore, Error, TEXT("Device: The device '%S' does support Multi Format but channel %d requested format's frame rate ('%S') and it's not compatible with the LTC reference frame rate ('%S')."), InCard->GetDisplayName().c_str(), uint32_t(InChannel) + 1, NTV2FrameRateToString(::GetNTV2FrameRateFromVideoFormat(OutFoundVideoFormat)).c_str(), NTV2FrameRateToString(LTCFrameRate).c_str());
 						return false;
 					}
 				}
@@ -1002,7 +1002,7 @@ namespace AJA
 				{
 					if (LTCFrameRate != ::GetNTV2FrameRateFromVideoFormat(OutFoundVideoFormat))
 					{
-						UE_LOG(LogTemp, Warning,  TEXT("Device: The device '%S' doesn't support Multi Format and channel %d requested format's frame rate ('%S') and it's not the same as the LTC reference frame rate ('%S')."), InCard->GetDisplayName().c_str(), uint32_t(InChannel) + 1, NTV2FrameRateToString(::GetNTV2FrameRateFromVideoFormat(OutFoundVideoFormat)).c_str(), NTV2FrameRateToString(LTCFrameRate).c_str());
+						UE_LOG(LogAjaCore, Warning,  TEXT("Device: The device '%S' doesn't support Multi Format and channel %d requested format's frame rate ('%S') and it's not the same as the LTC reference frame rate ('%S')."), InCard->GetDisplayName().c_str(), uint32_t(InChannel) + 1, NTV2FrameRateToString(::GetNTV2FrameRateFromVideoFormat(OutFoundVideoFormat)).c_str(), NTV2FrameRateToString(LTCFrameRate).c_str());
 						return false;
 					}
 				}
@@ -1016,13 +1016,13 @@ namespace AJA
 				{
 					if (!::IsMultiFormatCompatible(NotMultiVideoFormat, OutFoundVideoFormat))
 					{
-						UE_LOG(LogTemp, Error, TEXT("Device: The device '%S' does support Multi Format but channel %d requested format ('%S') and it's not compatible with the previous format ('%S')."), InCard->GetDisplayName().c_str(), uint32_t(InChannel) + 1, NTV2VideoFormatToString(OutFoundVideoFormat).c_str(), NTV2VideoFormatToString(NotMultiVideoFormat).c_str()); 
+						UE_LOG(LogAjaCore, Error, TEXT("Device: The device '%S' does support Multi Format but channel %d requested format ('%S') and it's not compatible with the previous format ('%S')."), InCard->GetDisplayName().c_str(), uint32_t(InChannel) + 1, NTV2VideoFormatToString(OutFoundVideoFormat).c_str(), NTV2VideoFormatToString(NotMultiVideoFormat).c_str()); 
 						return false;
 					}
 				}
 				else if (!bIsAutoDetected)
 				{
-					UE_LOG(LogTemp, Error, TEXT("Device: The device '%S' doesn't support Multi Format and channel %d requested format ('%S') which is not the same as previous format ('%S')."), InCard->GetDisplayName().c_str(), uint32_t(InChannel) + 1, NTV2VideoFormatToString(OutFoundVideoFormat).c_str(), NTV2VideoFormatToString(NotMultiVideoFormat).c_str());
+					UE_LOG(LogAjaCore, Error, TEXT("Device: The device '%S' doesn't support Multi Format and channel %d requested format ('%S') which is not the same as previous format ('%S')."), InCard->GetDisplayName().c_str(), uint32_t(InChannel) + 1, NTV2VideoFormatToString(OutFoundVideoFormat).c_str(), NTV2VideoFormatToString(NotMultiVideoFormat).c_str());
 					return false;
 				}
 			}
@@ -1052,13 +1052,13 @@ namespace AJA
 				std::string FailureReason;
 				if (!Helpers::GetInputVideoFormat(InCard, InTransportType, InChannel, InInputSource, InDesiredInputFormat, OutFoundVideoFormat, true, FailureReason))
 				{
-					UE_LOG(LogTemp, Error, TEXT("Device: Initialization of the input failed for channel %d on device '%S'. %S"), uint32_t(InChannel) + 1, InCard->GetDisplayName().c_str(), FailureReason.c_str());
+					UE_LOG(LogAjaCore, Error, TEXT("Device: Initialization of the input failed for channel %d on device '%S'. %S"), uint32_t(InChannel) + 1, InCard->GetDisplayName().c_str(), FailureReason.c_str());
 					return false;
 				}
 				
 				if (!Helpers::GetInputHDRMetadata(InCard, InChannel, InOutHDROptions))
 				{
-					UE_LOG(LogTemp, Error, TEXT("Device: Could not fetch HDR metadata from channel %d on device '%S'."), uint32_t(InChannel) + 1, InCard->GetDisplayName().c_str());
+					UE_LOG(LogAjaCore, Error, TEXT("Device: Could not fetch HDR metadata from channel %d on device '%S'."), uint32_t(InChannel) + 1, InCard->GetDisplayName().c_str());
 					return false;
 				}
 			}
@@ -1148,7 +1148,7 @@ namespace AJA
 				std::string FailureReason;
 				if (!Helpers::GetInputVideoFormat(InCard, InTransportType, InChannel, InInputSource, InDesiredInputFormat, OutFoundVideoFormat, true, FailureReason))
 				{
-					UE_LOG(LogTemp, Error, TEXT("Device: Initialization of the input failed for channel %d on device '%S'. %S"), uint32_t(InChannel) + 1, InCard->GetDisplayName().c_str(), FailureReason.c_str());
+					UE_LOG(LogAjaCore, Error, TEXT("Device: Initialization of the input failed for channel %d on device '%S'. %S"), uint32_t(InChannel) + 1, InCard->GetDisplayName().c_str(), FailureReason.c_str());
 					return false;
 				}
 			}
@@ -1196,7 +1196,7 @@ namespace AJA
 			int32_t  CurrentAppPID(0);
 			if (!Card->GetStreamingApplication(&CurrentAppFourCC, &CurrentAppPID))
 			{
-				UE_LOG(LogTemp, Error, TEXT("Device: Couldn't get current stream using Aja. Something is wrong when tryin to access registers on device '%S'."), Card->GetDisplayName().c_str());
+				UE_LOG(LogAjaCore, Error, TEXT("Device: Couldn't get current stream using Aja. Something is wrong when tryin to access registers on device '%S'."), Card->GetDisplayName().c_str());
 				delete Card;
 				Card = nullptr;
 				return false;
@@ -1222,11 +1222,11 @@ namespace AJA
 #endif
 					FourCCBuffer[4] = '\0';
 
-					UE_LOG(LogTemp, Error, TEXT("Device: Couldn't acquire stream for Unreal Engine application on device '%S'. Used by application '%S' with PID '%d'."), Card->GetDisplayName().c_str(), FourCCBuffer, CurrentAppPID);
+					UE_LOG(LogAjaCore, Error, TEXT("Device: Couldn't acquire stream for Unreal Engine application on device '%S'. Used by application '%S' with PID '%d'."), Card->GetDisplayName().c_str(), FourCCBuffer, CurrentAppPID);
 				}
 				else
 				{
-					UE_LOG(LogTemp, Error, TEXT("Device: Couldn't acquire stream for Unreal Engine application on device '%S'."), Card->GetDisplayName().c_str());
+					UE_LOG(LogAjaCore, Error, TEXT("Device: Couldn't acquire stream for Unreal Engine application on device '%S'."), Card->GetDisplayName().c_str());
 				}
 
 				delete Card;
@@ -1235,10 +1235,10 @@ namespace AJA
 			}
 			
 			{
-				UE_LOG(LogTemp,  Display, TEXT("Connected to %S with SDK %d.%d.%d.%d"), Card->GetDeviceVersionString().c_str(), AJA_NTV2_SDK_VERSION_MAJOR, AJA_NTV2_SDK_VERSION_MINOR, AJA_NTV2_SDK_VERSION_POINT, AJA_NTV2_SDK_BUILD_NUMBER);
+				UE_LOG(LogAjaCore,  Display, TEXT("Connected to %S with SDK %d.%d.%d.%d"), Card->GetDeviceVersionString().c_str(), AJA_NTV2_SDK_VERSION_MAJOR, AJA_NTV2_SDK_VERSION_MINOR, AJA_NTV2_SDK_VERSION_POINT, AJA_NTV2_SDK_BUILD_NUMBER);
 				if (AJA_NTV2_SDK_VERSION_MAJOR < 14)
 				{
-					UE_LOG(LogTemp, Warning,  TEXT("Unreal Engine's implementation of AJA support SDK 14 and up."));
+					UE_LOG(LogAjaCore, Warning,  TEXT("Unreal Engine's implementation of AJA support SDK 14 and up."));
 				}
 			}
 
@@ -1268,7 +1268,7 @@ namespace AJA
 				}
 				else
 				{
-					UE_LOG(LogTemp, Warning,  TEXT("Device: The device '%S' doesn't support Multi Format."), Card->GetDisplayName().c_str());
+					UE_LOG(LogAjaCore, Warning,  TEXT("Device: The device '%S' doesn't support Multi Format."), Card->GetDisplayName().c_str());
 				}
 			}
 
@@ -1306,7 +1306,7 @@ namespace AJA
 			}
 			else
 			{
-				UE_LOG(LogTemp, Error, TEXT("Device: Too many Frames Indexes requested on device '%S'.\n"), Card->GetDisplayName().c_str());
+				UE_LOG(LogAjaCore, Error, TEXT("Device: Too many Frames Indexes requested on device '%S'.\n"), Card->GetDisplayName().c_str());
 				return InvalidFrameBufferIndex;
 			}
 		}
diff --git a/Engine/Plugins/Media/AjaMedia/Source/Aja/Private/Helpers.cpp b/Engine/Plugins/Media/AjaMedia/Source/Aja/Private/Helpers.cpp
--- a/Engine/Plugins/Media/AjaMedia/Source/Aja/Private/Helpers.cpp
+++ b/Engine/Plugins/Media/AjaMedia/Source/Aja/Private/Helpers.cpp
@@ -315,7 +315,7 @@ namespace AJA
 			{
 				if (bInLogError)
 				{
-					UE_LOG(LogTemp, Warning,  TEXT("AJA: The timecode type is not present for LTC Input '%d' on device '%S'.\n"), uint32_t(AnalogLTCInput)+1, InCard->GetDisplayName().c_str());
+					UE_LOG(LogAjaCore, Warning,  TEXT("AJA: The timecode type is not present for LTC Input '%d' on device '%S'.\n"), uint32_t(AnalogLTCInput)+1, InCard->GetDisplayName().c_str());
 				}
 				return false;
 			}
@@ -325,7 +325,7 @@ namespace AJA
 			{
 				if (bInLogError)
 				{
-					UE_LOG(LogTemp, Warning,  TEXT("AJA: The timecode type is not present for LTC Input '%d' on device '%S'.\n"), uint32_t(AnalogLTCInput) + 1, InCard->GetDisplayName().c_str());
+					UE_LOG(LogAjaCore, Warning,  TEXT("AJA: The timecode type is not present for LTC Input '%d' on device '%S'.\n"), uint32_t(AnalogLTCInput) + 1, InCard->GetDisplayName().c_str());
 				}
 				return false;
 			}
@@ -637,7 +637,7 @@ namespace AJA
 				}
 				else
 				{
-					UE_LOG(LogTemp, Warning,  TEXT("RouteSignal: This input routing is not supported. %S.\n"), Helpers::TransportTypeToString(InTransportType));
+					UE_LOG(LogAjaCore, Warning,  TEXT("RouteSignal: This input routing is not supported. %S.\n"), Helpers::TransportTypeToString(InTransportType));
 				}
 			}
 			else //is output
@@ -751,7 +751,7 @@ namespace AJA
 				}
 				else
 				{
-					UE_LOG(LogTemp, Warning,  TEXT("RouteSignal: This output routing is not supported. %S.\n"), Helpers::TransportTypeToString(InTransportType));
+					UE_LOG(LogAjaCore, Warning,  TEXT("RouteSignal: This output routing is not supported. %S.\n"), Helpers::TransportTypeToString(InTransportType));
 				}
 			}
 		}
@@ -940,7 +940,7 @@ namespace AJA
 			}
 			else
 			{
-				UE_LOG(LogTemp, Warning,  TEXT("RouteKeySignal: Key routing is not supported. %S.\n"), Helpers::TransportTypeToString(InTransportType));
+				UE_LOG(LogAjaCore, Warning,  TEXT("RouteKeySignal: Key routing is not supported. %S.\n"), Helpers::TransportTypeToString(InTransportType));
 			}
 		}
 
@@ -989,7 +989,7 @@ namespace AJA
 			{
 				if (InInputSource != NTV2_INPUTSOURCE_HDMI1 && InInputSource != NTV2_INPUTSOURCE_HDMI2)
 				{
-					UE_LOG(LogTemp, Warning,  TEXT("AJA: 4K HDMI is only supported in inputs 1 and 2.\n"));
+					UE_LOG(LogAjaCore, Warning,  TEXT("AJA: 4K HDMI is only supported in inputs 1 and 2.\n"));
 					return NTV2_OUTPUT_CROSSPOINT_INVALID;
 				}
 			}
@@ -1038,7 +1038,7 @@ namespace AJA
 					NTV2DeviceInfo DeviceInfo;
 					if (!Scanner.GetDeviceInfo(DeviceIndex, DeviceInfo, false))
 					{
-						UE_LOG(LogTemp, Error, TEXT("ConfigureVideo: Device not found.\n"));
+						UE_LOG(LogAjaCore, Error, TEXT("ConfigureVideo: Device not found.\n"));
 						return false;
 					}
 					DeviceID = DeviceInfo.deviceID;
@@ -1356,7 +1356,7 @@ namespace AJA
 			{
 				if (bInLogError)
 				{
-					UE_LOG(LogTemp, Warning,  TEXT("AJA: There is no timecode present for channel '%d' on device '%S'.\n"), uint32_t(InChannel) + 1, InCard->GetDisplayName().c_str());
+					UE_LOG(LogAjaCore, Warning,  TEXT("AJA: There is no timecode present for channel '%d' on device '%S'.\n"), uint32_t(InChannel) + 1, InCard->GetDisplayName().c_str());
 				}
 				return false;
 			}
@@ -1372,7 +1372,7 @@ namespace AJA
 					{
 						if (bInLogError)
 						{
-							UE_LOG(LogTemp, Warning,  TEXT("AJA: The timecode type is not present for channel '%d' on device '%S'.\n"), uint32_t(InChannel) + 1, InCard->GetDisplayName().c_str());
+							UE_LOG(LogAjaCore, Warning,  TEXT("AJA: The timecode type is not present for channel '%d' on device '%S'.\n"), uint32_t(InChannel) + 1, InCard->GetDisplayName().c_str());
 						}
 					}
 					break;
@@ -1384,7 +1384,7 @@ namespace AJA
 					{
 						if (bInLogError)
 						{
-							UE_LOG(LogTemp, Warning,  TEXT("AJA: The timecode type is not present for channel '%d' on device '%S'.\n"), uint32_t(InChannel) + 1, InCard->GetDisplayName().c_str());
+							UE_LOG(LogAjaCore, Warning,  TEXT("AJA: The timecode type is not present for channel '%d' on device '%S'.\n"), uint32_t(InChannel) + 1, InCard->GetDisplayName().c_str());
 						}
 					}
 					break;
diff --git a/Engine/Plugins/Media/AjaMedia/Source/Aja/Private/Helpers.h b/Engine/Plugins/Media/AjaMedia/Source/Aja/Private/Helpers.h
--- a/Engine/Plugins/Media/AjaMedia/Source/Aja/Private/Helpers.h
+++ b/Engine/Plugins/Media/AjaMedia/Source/Aja/Private/Helpers.h
@@ -3,6 +3,7 @@
 #pragma once
 
 #include "CoreMinimal.h"
+#include "AjaCoreModule.h"
 
 #define MSWindows 1
 #define AJA_WINDOWS 1
diff --git a/Engine/Plugins/Media/AjaMedia/Source/Aja/Private/OutputChannel.cpp b/Engine/Plugins/Media/AjaMedia/Source/Aja/Private/OutputChannel.cpp
--- a/Engine/Plugins/Media/AjaMedia/Source/Aja/Private/OutputChannel.cpp
+++ b/Engine/Plugins/Media/AjaMedia/Source/Aja/Private/OutputChannel.cpp
@@ -492,7 +492,7 @@ namespace AJA
 
 			if (InAncillaryBuffer == nullptr || InAncillaryBufferSize > AncBufferSize || (!bIsProgressive && InAncillaryBufferSize > AncF2BufferSize))
 			{
-				UE_LOG(LogTemp, Error, TEXT("SetFrameData: Can't set the ancillary. The buffer is invalid or the buffer size is not the same as the AJA for output channel %d on device %S.\n"), uint32_t(Channel) + 1, GetDevice().GetDisplayName().c_str());
+				UE_LOG(LogAjaCore, Error, TEXT("SetFrameData: Can't set the ancillary. The buffer is invalid or the buffer size is not the same as the AJA for output channel %d on device %S.\n"), uint32_t(Channel) + 1, GetDevice().GetDisplayName().c_str());
 				return false;
 			}
 
@@ -550,7 +550,7 @@ namespace AJA
 			
 			if (InAudioBufferSize > AudioBufferSize)
 			{
-				UE_LOG(LogTemp, Error, TEXT("SetAudioFrameData: Can't set the audio. The buffer size is not the same as the AJA Audio Format for output channel %d on device %S.\n"), uint32_t(Channel) + 1, GetDevice().GetDisplayName().c_str());
+				UE_LOG(LogAjaCore, Error, TEXT("SetAudioFrameData: Can't set the audio. The buffer size is not the same as the AJA Audio Format for output channel %d on device %S.\n"), uint32_t(Channel) + 1, GetDevice().GetDisplayName().c_str());
 				return false;
 			}
 
@@ -587,7 +587,7 @@ namespace AJA
 
 			if (InVideoBuffer == nullptr || InVideoBufferSize > VideoBufferSize)
 			{
-				UE_LOG(LogTemp, Error, TEXT("SetVideoFrameData: Can't set the video. The buffer is invalid or the buffer size is not the same as the AJA Video Format for output channel %d on device %S.\n"), uint32_t(Channel) + 1, GetDevice().GetDisplayName().c_str());
+				UE_LOG(LogAjaCore, Error, TEXT("SetVideoFrameData: Can't set the video. The buffer is invalid or the buffer size is not the same as the AJA Video Format for output channel %d on device %S.\n"), uint32_t(Channel) + 1, GetDevice().GetDisplayName().c_str());
 				return false;
 			}
 
@@ -654,7 +654,7 @@ namespace AJA
 
 			if (!bIsProgressive)
 			{
-				UE_LOG(LogTemp, Error, TEXT("SetVideoFrameData: Can't set the video, GPUTextureTransfer is not supported with interlaced video for output channel %d on device %S.\n"), uint32_t(Channel) + 1, GetDevice().GetDisplayName().c_str());
+				UE_LOG(LogAjaCore, Error, TEXT("SetVideoFrameData: Can't set the video, GPUTextureTransfer is not supported with interlaced video for output channel %d on device %S.\n"), uint32_t(Channel) + 1, GetDevice().GetDisplayName().c_str());
 				return false;
 			}
 
@@ -665,13 +665,13 @@ namespace AJA
 
 			if (!RHITexture)
 			{
-				UE_LOG(LogTemp, Error, TEXT("SetVideoFrameData: Can't set the video. The RHI texture is invalid for output channel %d on device %S.\n"), uint32_t(Channel) + 1, GetDevice().GetDisplayName().c_str());
+				UE_LOG(LogAjaCore, Error, TEXT("SetVideoFrameData: Can't set the video. The RHI texture is invalid for output channel %d on device %S.\n"), uint32_t(Channel) + 1, GetDevice().GetDisplayName().c_str());
 				return false;
 			}
 
 			if (!TextureTransfer || !GetOptions().bUseGPUDMA)
 			{
-				UE_LOG(LogTemp, Error, TEXT("SetVideoFrameData: Can't set the video. GPU DMA was not setup correctly for output channel %d on device %S.\n"), uint32_t(Channel) + 1, GetDevice().GetDisplayName().c_str());
+				UE_LOG(LogAjaCore, Error, TEXT("SetVideoFrameData: Can't set the video. GPU DMA was not setup correctly for output channel %d on device %S.\n"), uint32_t(Channel) + 1, GetDevice().GetDisplayName().c_str());
 				return false;
 			}
 
@@ -785,7 +785,7 @@ namespace AJA
 						bResult &= GetDevice().StartAudioOutput(AudioSystem);
 						if (bResult == false)
 						{
-							UE_LOG(LogTemp, Error, TEXT("ConfigureAudio: Could not start audio output for audio system %d for channel %d on device %S."), uint32_t(AudioSystem), uint32_t(Channel) + 1, GetDevice().GetDisplayName().c_str());
+							UE_LOG(LogAjaCore, Error, TEXT("ConfigureAudio: Could not start audio output for audio system %d for channel %d on device %S."), uint32_t(AudioSystem), uint32_t(Channel) + 1, GetDevice().GetDisplayName().c_str());
 						}
 					}
 				}
@@ -882,14 +882,14 @@ namespace AJA
 
 				if (!bRunning)
 				{
-					UE_LOG(LogTemp, Error, TEXT("AutoCirculate: Can't wait for the output field for channel %d on device %S.\n"), uint32_t(Channel) + 1, GetDevice().GetDisplayName().c_str());
+					UE_LOG(LogAjaCore, Error, TEXT("AutoCirculate: Can't wait for the output field for channel %d on device %S.\n"), uint32_t(Channel) + 1, GetDevice().GetDisplayName().c_str());
 				}
 
 				AUTOCIRCULATE_STATUS ChannelStatus;
 				bRunning = GetDevice().AutoCirculateGetStatus(Channel, ChannelStatus);
 				if (!bRunning)
 				{
-					UE_LOG(LogTemp, Error, TEXT("AutoCirculate: Can't get the status for output channel %d on device %S.\n"), uint32_t(Channel) + 1, GetDevice().GetDisplayName().c_str());
+					UE_LOG(LogAjaCore, Error, TEXT("AutoCirculate: Can't get the status for output channel %d on device %S.\n"), uint32_t(Channel) + 1, GetDevice().GetDisplayName().c_str());
 					break;
 				}
 
@@ -1126,7 +1126,7 @@ namespace AJA
 
 				if (!bRunning)
 				{
-					UE_LOG(LogTemp, Error, TEXT("PingPong: Can't wait for the output field for channel %d on device %S.\n"), uint32_t(Channel) + 1, GetDevice().GetDisplayName().c_str());
+					UE_LOG(LogAjaCore, Error, TEXT("PingPong: Can't wait for the output field for channel %d on device %S.\n"), uint32_t(Channel) + 1, GetDevice().GetDisplayName().c_str());
 					break;
 				}
 
@@ -1198,7 +1198,7 @@ namespace AJA
 
 					if (!bRunning)
 					{
-						UE_LOG(LogTemp, Error, TEXT("PingPong: Can't do the DMA frame transfer for channel %d on device %S.\n"), uint32_t(Channel) + 1, GetDevice().GetDisplayName().c_str());
+						UE_LOG(LogAjaCore, Error, TEXT("PingPong: Can't do the DMA frame transfer for channel %d on device %S.\n"), uint32_t(Channel) + 1, GetDevice().GetDisplayName().c_str());
 						Thread_PushAvailableReadingFrame(AvailableReadingFrame);
 						break;
 					}
@@ -1244,7 +1244,7 @@ namespace AJA
 					}
 
 					TRACE_CPUPROFILER_EVENT_SCOPE(OutputChannel::NoFramesAvailable);
-					UE_LOG(LogTemp, Warning,  TEXT("PingPong: No frames are available for channel %d on device %S.\n"), uint32_t(Channel) + 1, GetDevice().GetDisplayName().c_str());
+					UE_LOG(LogAjaCore, Warning,  TEXT("PingPong: No frames are available for channel %d on device %S.\n"), uint32_t(Channel) + 1, GetDevice().GetDisplayName().c_str());
 				}
 			}
 
@@ -1301,12 +1301,12 @@ namespace AJA
 			bool pIsSecondLineWhite = IsLineWhite(pSecondLine);
 			if (bIsFirstLineWhite == pIsSecondLineWhite)
 			{
-				UE_LOG(LogTemp, Error, TEXT("INTERLACED TEST - The 2 lines are the same color. %d\n"), InterlacedTest_FrameCounter);
+				UE_LOG(LogAjaCore, Error, TEXT("INTERLACED TEST - The 2 lines are the same color. %d\n"), InterlacedTest_FrameCounter);
 			}
 			if (bIsFirstLineWhite != bInterlacedTest_ExpectFirstLineToBeWhite)
 			{
 				bInterlacedTest_ExpectFirstLineToBeWhite = !bInterlacedTest_ExpectFirstLineToBeWhite;
-				UE_LOG(LogTemp, Error, TEXT("INTERLACED TEST - The lines has swap color. %d\n"), InterlacedTest_FrameCounter);
+				UE_LOG(LogAjaCore, Error, TEXT("INTERLACED TEST - The lines has swap color. %d\n"), InterlacedTest_FrameCounter);
 			}
 			++InterlacedTest_FrameCounter;
 		}
@@ -1410,7 +1410,7 @@ namespace AJA
 
 			if (bIsPaused && bAudioOutputRunning && !Thread_ShouldPauseAudioOutput())
 			{
-				UE_LOG(LogTemp, Verbose, TEXT("%s: Resuming audio output."), *OutputMethod);
+				UE_LOG(LogAjaCore, Verbose, TEXT("%s: Resuming audio output."), *OutputMethod);
 				bool bPausePlayback = false;
 				bAudioSuccess &= GetDevice().SetAudioOutputPause(AudioSystem, bPausePlayback);
 				bIsPaused = true;
@@ -1427,7 +1427,7 @@ namespace AJA
 
 				CurrentAudioWriteOffset = InitialOffset;
 
-				UE_LOG(LogTemp, Verbose, TEXT("%s: Starting audio output."), *OutputMethod);
+				UE_LOG(LogAjaCore, Verbose, TEXT("%s: Starting audio output."), *OutputMethod);
 				bAudioSuccess &= GetDevice().StartAudioOutput(AudioSystem); // This resets the playhead to 0, we have to override that to something closer to the write head in order to reduce audio delay
 			}
 
@@ -1440,7 +1440,7 @@ namespace AJA
 				AudioPlayheadLastPosition = PlayHeadPosition;
 				
 				SET_FLOAT_STAT(STAT_AjaMediaCapture_Audio_Delay, OffsetInSeconds);
-				UE_LOG(LogTemp, Verbose, TEXT("%s: Play head: %d, Write Head: %d, Offset: %f"), *OutputMethod, AudioPlayheadLastPosition.load(), CurrentAudioWriteOffset.load(), OffsetInSeconds);
+				UE_LOG(LogAjaCore, Verbose, TEXT("%s: Play head: %d, Write Head: %d, Offset: %f"), *OutputMethod, AudioPlayheadLastPosition.load(), CurrentAudioWriteOffset.load(), OffsetInSeconds);
 
 				bAudioTransferWasSuccessful &= Thread_TransferAudioBuffer(AvailableReadingFrame->AudioBuffer, AvailableReadingFrame->CopiedAudioBufferSize);
 
@@ -1449,17 +1449,17 @@ namespace AJA
 
 				if (bAudioSuccess && Thread_ShouldPauseAudioOutput())
 				{
-					UE_LOG(LogTemp, Warning, TEXT("%s: Audio play head is catching up to write head for channel %d on device %S.\n"), *OutputMethod, uint32_t(Channel) + 1, GetDevice().GetDisplayName().c_str());
+					UE_LOG(LogAjaCore, Warning, TEXT("%s: Audio play head is catching up to write head for channel %d on device %S.\n"), *OutputMethod, uint32_t(Channel) + 1, GetDevice().GetDisplayName().c_str());
 
 					bool bPausePlayout = true;
 					GetDevice().SetAudioOutputPause(AudioSystem, bPausePlayout);
 
-					UE_LOG(LogTemp, Verbose, TEXT("%s Pausing audio output because the play head is catching up to the write head. Last playhead position: %d"), *OutputMethod, AudioPlayheadLastPosition.load());
+					UE_LOG(LogAjaCore, Verbose, TEXT("%s Pausing audio output because the play head is catching up to the write head. Last playhead position: %d"), *OutputMethod, AudioPlayheadLastPosition.load());
 				}
 			}
 			else
 			{
-				UE_LOG(LogTemp, Verbose, TEXT("%s Audio output failed."), *OutputMethod);
+				UE_LOG(LogAjaCore, Verbose, TEXT("%s Audio output failed."), *OutputMethod);
 			}
 
 			return bAudioTransferWasSuccessful;
@@ -1480,12 +1480,12 @@ namespace AJA
 			bool bAudioOutputRunning = false;
 			bAudioSuccess &= GetDevice().IsAudioOutputRunning(AudioSystem, bAudioOutputRunning);
 
-			UE_CLOG(!bAudioSuccess, LogTemp, Verbose, TEXT("IsAudioOutputRunning returned false."));
+			UE_CLOG(!bAudioSuccess, LogAjaCore, Verbose, TEXT("IsAudioOutputRunning returned false."));
 
 			bool bIsPaused = false;
 			bAudioSuccess &= GetDevice().GetAudioOutputPause(AudioSystem, bIsPaused);
 
-			UE_CLOG(!bAudioSuccess, LogTemp, Verbose, TEXT("GetAudioOutputPause returned false."));
+			UE_CLOG(!bAudioSuccess, LogAjaCore, Verbose, TEXT("GetAudioOutputPause returned false."));
 
 			// Start audio output as late as possible to reduce delay.
 			if (!bAudioOutputRunning)
@@ -1499,20 +1499,20 @@ namespace AJA
 
 				CurrentAudioWriteOffset = InitialOffset;
 
-				UE_LOG(LogTemp, Verbose, TEXT("PingPong: Starting audio output."));
+				UE_LOG(LogAjaCore, Verbose, TEXT("PingPong: Starting audio output."));
 				bAudioSuccess &= GetDevice().StartAudioOutput(AudioSystem); // This resets the playhead to 0, we have to override that to something closer to the write head in order to reduce audio delay
 			}
 
 			double OffsetInSeconds = 0;
 			bAudioSuccess &= Thread_GetAudioOffsetInSeconds(OffsetInSeconds);
-			UE_CLOG(!bAudioSuccess, LogTemp, Verbose, TEXT("Thread_GetAudioOffsetInSeconds returned false."));
+			UE_CLOG(!bAudioSuccess, LogAjaCore, Verbose, TEXT("Thread_GetAudioOffsetInSeconds returned false."));
 			ULWord PlayHeadPosition = 0;
 			bAudioSuccess &= GetDevice().ReadAudioLastOut(PlayHeadPosition, AudioSystem);
-			UE_CLOG(!bAudioSuccess, LogTemp, Verbose, TEXT("ReadAudioLastOut returned false."));
+			UE_CLOG(!bAudioSuccess, LogAjaCore, Verbose, TEXT("ReadAudioLastOut returned false."));
 			
 			AudioPlayheadLastPosition = PlayHeadPosition;
 			SET_FLOAT_STAT(STAT_AjaMediaCapture_Audio_Delay, OffsetInSeconds);
-			UE_LOG(LogTemp, Verbose, TEXT("%s: Play head: %d, Write Head: %d, Offset: %f"), TEXT("PingPong"), AudioPlayheadLastPosition.load(), CurrentAudioWriteOffset.load(), OffsetInSeconds);
+			UE_LOG(LogAjaCore, Verbose, TEXT("%s: Play head: %d, Write Head: %d, Offset: %f"), TEXT("PingPong"), AudioPlayheadLastPosition.load(), CurrentAudioWriteOffset.load(), OffsetInSeconds);
 
 			bAudioSuccess &= Thread_TransferAudioBuffer(InAudioBuffer, BufferSize);
 			
@@ -1536,7 +1536,7 @@ namespace AJA
 				{
 					constexpr bool bPausePlayout = true;
 					bAudioSuccess &= GetDevice().SetAudioOutputPause(AudioSystem, bPausePlayout);
-					UE_LOG(LogTemp, Warning, TEXT("Pausing audio output. Last playhead position: %d"), AudioPlayheadLastPosition.load());
+					UE_LOG(LogAjaCore, Warning, TEXT("Pausing audio output. Last playhead position: %d"), AudioPlayheadLastPosition.load());
 				}
 			}
 
diff --git a/Engine/Plugins/Media/AjaMedia/Source/Aja/Private/SyncChannel.cpp b/Engine/Plugins/Media/AjaMedia/Source/Aja/Private/SyncChannel.cpp
--- a/Engine/Plugins/Media/AjaMedia/Source/Aja/Private/SyncChannel.cpp
+++ b/Engine/Plugins/Media/AjaMedia/Source/Aja/Private/SyncChannel.cpp
@@ -141,7 +141,7 @@ namespace AJA
 		{
 			if (InOptions.ChannelIndex > NTV2_MAX_NUM_CHANNELS || InOptions.ChannelIndex < 1)
 			{
-				UE_LOG(LogTemp, Error, TEXT("SyncChannel: The port index '%d' is invalid.\n"), InOptions.ChannelIndex);
+				UE_LOG(LogAjaCore, Error, TEXT("SyncChannel: The port index '%d' is invalid.\n"), InOptions.ChannelIndex);
 			}
 			else
 			{
@@ -244,7 +244,7 @@ namespace AJA
 
 			if (!Helpers::TryVideoFormatIndexToNTV2VideoFormat(SyncOption.VideoFormatIndex, DesiredVideoFormat))
 			{
-				UE_LOG(LogTemp, Error, TEXT("SyncChannel: The expected video format is invalid for %d.\n"), uint32_t(Channel) + 1);
+				UE_LOG(LogAjaCore, Error, TEXT("SyncChannel: The expected video format is invalid for %d.\n"), uint32_t(Channel) + 1);
 				return false;
 			}
 
@@ -279,7 +279,7 @@ namespace AJA
 
 			if (!Device->WaitForInputOrOutputInterrupt(Channel, 12 * 2))
 			{
-				UE_LOG(LogTemp, Error, TEXT("SyncChannel: Was not able to lock on device %S for channel %d.\n"), GetDevice().GetDisplayName().c_str(), uint32_t(Channel) + 1);
+				UE_LOG(LogAjaCore, Error, TEXT("SyncChannel: Was not able to lock on device %S for channel %d.\n"), GetDevice().GetDisplayName().c_str(), uint32_t(Channel) + 1);
 				return false;
 			}
 
@@ -301,7 +301,7 @@ namespace AJA
 
 				if (!FoundFormat.has_value())
 				{
-					UE_LOG(LogTemp, Error, TEXT("Sync: The VideoFormat was invalid for channel %d on device %S. %S\n")
+					UE_LOG(LogAjaCore, Error, TEXT("Sync: The VideoFormat was invalid for channel %d on device %S. %S\n")
 						, uint32_t(Channel) + 1
 						, GetDevice().GetDisplayName().c_str()
 						, FailureReason.c_str()
@@ -318,7 +318,7 @@ namespace AJA
 					}
 					else
 					{
-						UE_LOG(LogTemp, Error, TEXT("Sync: The VideoFormat changed for channel %d on device %S. %S\n")
+						UE_LOG(LogAjaCore, Error, TEXT("Sync: The VideoFormat changed for channel %d on device %S. %S\n")
 							, uint32_t(Channel) + 1
 							, GetDevice().GetDisplayName().c_str()
 							, FailureReason.c_str()
diff --git a/Engine/Plugins/Media/AjaMedia/Source/Aja/Private/TimecodeChannel.cpp b/Engine/Plugins/Media/AjaMedia/Source/Aja/Private/TimecodeChannel.cpp
--- a/Engine/Plugins/Media/AjaMedia/Source/Aja/Private/TimecodeChannel.cpp
+++ b/Engine/Plugins/Media/AjaMedia/Source/Aja/Private/TimecodeChannel.cpp
@@ -115,7 +115,7 @@ namespace AJA
 			{
 				if (InOptions.ChannelIndex > NTV2_MAX_NUM_CHANNELS || InOptions.ChannelIndex < 1)
 				{
-					UE_LOG(LogTemp, Error, TEXT("SyncChannel: The port index '%d' is invalid.\n"), InOptions.ChannelIndex);
+					UE_LOG(LogAjaCore, Error, TEXT("SyncChannel: The port index '%d' is invalid.\n"), InOptions.ChannelIndex);
 				}
 				else
 				{
@@ -227,7 +227,7 @@ namespace AJA
 				case 1: LtcSource = EAnalogLTCSource::LTC1; break;
 				case 2: LtcSource = EAnalogLTCSource::LTC2; break;
 				default:
-					UE_LOG(LogTemp, Error, TEXT("TimecodeChannel: The LTC source Index is invalid on device %S.\n"), GetDevice().GetDisplayName().c_str());
+					UE_LOG(LogAjaCore, Error, TEXT("TimecodeChannel: The LTC source Index is invalid on device %S.\n"), GetDevice().GetDisplayName().c_str());
 					return false;
 				}
 
@@ -244,7 +244,7 @@ namespace AJA
 
 				if (!Helpers::TryVideoFormatIndexToNTV2VideoFormat(TimecodeOption.VideoFormatIndex, DesiredVideoFormat))
 				{
-					UE_LOG(LogTemp, Error, TEXT("TimecodeChannel: The expected video format is invalid for %d.\n"), uint32_t(Channel) + 1);
+					UE_LOG(LogAjaCore, Error, TEXT("TimecodeChannel: The expected video format is invalid for %d.\n"), uint32_t(Channel) + 1);
 					return false;
 				}
 
@@ -270,7 +270,7 @@ namespace AJA
 
 				if (!Device->WaitForInputOrOutputInterrupt(Channel, 12 * 2))
 				{
-					UE_LOG(LogTemp, Error, TEXT("TimecodeChannel: Was not able to lock on device %S for channel %d.\n"), GetDevice().GetDisplayName().c_str(), uint32_t(Channel) + 1);
+					UE_LOG(LogAjaCore, Error, TEXT("TimecodeChannel: Was not able to lock on device %S for channel %d.\n"), GetDevice().GetDisplayName().c_str(), uint32_t(Channel) + 1);
 					return false;
 				}
 			}
diff --git a/Engine/Plugins/Media/AjaMedia/Source/Aja/Public/AjaCoreModule.h b/Engine/Plugins/Media/AjaMedia/Source/Aja/Public/AjaCoreModule.h
--- a/Engine/Plugins/Media/AjaMedia/Source/Aja/Public/AjaCoreModule.h
+++ b/Engine/Plugins/Media/AjaMedia/Source/Aja/Public/AjaCoreModule.h
@@ -5,6 +5,8 @@
 #include "Modules/ModuleInterface.h"
 #include "Modules/ModuleManager.h"
 
+DECLARE_LOG_CATEGORY_EXTERN(LogAjaCore, Log, All);
+
 class FAjaCoreModule : public IModuleInterface
 {
 };
