Allow setting aja genlock channel to same output as video

diff --git a/Engine/Plugins/Media/AjaMedia/Source/Aja/Private/ChannelBase.cpp b/Engine/Plugins/Media/AjaMedia/Source/Aja/Private/ChannelBase.cpp
--- a/Engine/Plugins/Media/AjaMedia/Source/Aja/Private/ChannelBase.cpp
+++ b/Engine/Plugins/Media/AjaMedia/Source/Aja/Private/ChannelBase.cpp
@@ -299,6 +299,7 @@ namespace AJA
 			
 			const bool bConnectChannel = true;
 			const bool bAsOwner = true;
+			const bool bAsGenlock = false;
 
 			// Get the override channel in case we need to reroute an input to a different channel.
 			// ie. HDMI2 4K should go to channel 3
@@ -311,7 +312,7 @@ namespace AJA
 				DesiredVideoFormat = Helpers::GetLevelA(DesiredVideoFormat);
 			}
 
-			if (!InCommandList.RegisterChannel(GetOptions().TransportType, InputSource, Channel, IsInput(), bConnectChannel, GetOptions().TimecodeFormat, GetOptions().PixelFormat, DesiredVideoFormat, Options.HDROptions, bAsOwner, Options.bAutoDetectFormat))
+			if (!InCommandList.RegisterChannel(GetOptions().TransportType, InputSource, Channel, IsInput(), bAsGenlock, bConnectChannel, GetOptions().TimecodeFormat, GetOptions().PixelFormat, DesiredVideoFormat, Options.HDROptions, bAsOwner, Options.bAutoDetectFormat))
 			{
 				return false;
 			}
@@ -547,8 +548,10 @@ namespace AJA
 					if (UseKey())
 					{
 						bConnectedChannel = false;
-						const bool bAsOwner = true;
-						if (!InCommandList.RegisterChannel(GetOptions().TransportType, KeySource, KeyChannel, IsInput(), bConnectedChannel, ETimecodeFormat::TCF_None, GetOptions().PixelFormat, DesiredVideoFormat, Options.HDROptions, bAsOwner, Options.bAutoDetectFormat))
+						constexpr bool bAsOwner = true;
+						constexpr bool bAsGenlock = false;
+
+						if (!InCommandList.RegisterChannel(GetOptions().TransportType, KeySource, KeyChannel, IsInput(), bAsGenlock, bConnectedChannel, ETimecodeFormat::TCF_None, GetOptions().PixelFormat, DesiredVideoFormat, Options.HDROptions, bAsOwner, Options.bAutoDetectFormat))
 						{
 							return false;
 						}
@@ -588,8 +591,10 @@ namespace AJA
 					if (UseKey())
 					{
 						const bool bUseKeyConnectedChannel = false;
-						const bool bAsOwner = true;
-						if (!InCommandList.RegisterChannel(GetOptions().TransportType, KeySource, KeyChannel, IsInput(), bUseKeyConnectedChannel, ETimecodeFormat::TCF_None, GetOptions().PixelFormat, DesiredVideoFormat, Options.HDROptions, bAsOwner, Options.bAutoDetectFormat))
+						constexpr bool bAsOwner = true;
+						constexpr bool bAsGenlock = false;
+
+						if (!InCommandList.RegisterChannel(GetOptions().TransportType, KeySource, KeyChannel, IsInput(), bAsGenlock, bUseKeyConnectedChannel, ETimecodeFormat::TCF_None, GetOptions().PixelFormat, DesiredVideoFormat, Options.HDROptions, bAsOwner, Options.bAutoDetectFormat))
 						{
 							return false;
 						}
diff --git a/Engine/Plugins/Media/AjaMedia/Source/Aja/Private/Device.cpp b/Engine/Plugins/Media/AjaMedia/Source/Aja/Private/Device.cpp
--- a/Engine/Plugins/Media/AjaMedia/Source/Aja/Private/Device.cpp
+++ b/Engine/Plugins/Media/AjaMedia/Source/Aja/Private/Device.cpp
@@ -153,6 +153,7 @@ namespace AJA
 			, bIsOwned(false)
 			, bIsInput(false)
 			, bIsAutoDetected(false)
+			, bGenlockChannel(false)
 		{
 		}
 
@@ -162,13 +163,13 @@ namespace AJA
 			: Connection(InConnection)
 		{}
 
-		bool DeviceConnection::CommandList::RegisterChannel(ETransportType InTransportType, NTV2InputSource InInputSource, NTV2Channel InChannel, bool bInAsInput, bool bConnectChannel, ETimecodeFormat InTimecodeFormat, EPixelFormat InUEPixelFormat, NTV2VideoFormat InDesiredInputFormat, bool bInAsOwner, bool bInIsAutoDetected)
+		bool DeviceConnection::CommandList::RegisterChannel(ETransportType InTransportType, NTV2InputSource InInputSource, NTV2Channel InChannel, bool bInAsInput, bool bInAsGenlock, bool bConnectChannel, ETimecodeFormat InTimecodeFormat, EPixelFormat InUEPixelFormat, NTV2VideoFormat InDesiredInputFormat, bool bInAsOwner, bool bInIsAutoDetected)
 		{
 			FAjaHDROptions UnusedOptions;
-			return RegisterChannel(InTransportType, InInputSource, InChannel, bInAsInput, bConnectChannel, InTimecodeFormat, InUEPixelFormat, InDesiredInputFormat, UnusedOptions, bInAsOwner, bInIsAutoDetected);
+			return RegisterChannel(InTransportType, InInputSource, InChannel, bInAsInput, bInAsGenlock, bConnectChannel, InTimecodeFormat, InUEPixelFormat, InDesiredInputFormat, UnusedOptions, bInAsOwner, bInIsAutoDetected);
 		}
 
-		bool DeviceConnection::CommandList::RegisterChannel(ETransportType InTransportType, NTV2InputSource InInputSource, NTV2Channel InChannel, bool bInIsInput, bool bConnectChannel, ETimecodeFormat InTimecodeFormat, EPixelFormat InUEPixelFormat, NTV2VideoFormat InDesiredInputFormat, FAjaHDROptions& InOutHDROptions, bool bInAsOwner, bool bInIsAutoDetected)
+		bool DeviceConnection::CommandList::RegisterChannel(ETransportType InTransportType, NTV2InputSource InInputSource, NTV2Channel InChannel, bool bInIsInput, bool bInAsGenlock, bool bConnectChannel, ETimecodeFormat InTimecodeFormat, EPixelFormat InUEPixelFormat, NTV2VideoFormat InDesiredInputFormat, FAjaHDROptions& InOutHDROptions, bool bInAsOwner, bool bInIsAutoDetected)
 		{
 			AJAAutoLock Lock(&Connection.ChannelLock);
 
@@ -225,6 +226,7 @@ namespace AJA
 					Info->bIsInput = bInIsInput;
 					Info->bConnected = bConnectChannel;
 					Info->bIsAutoDetected = bInIsAutoDetected;
+					Info->bGenlockChannel = bInAsGenlock;
 
 					{
 						AJAAutoLock AutoLock(&Connection.ChannelInfoLock);
@@ -283,7 +285,11 @@ namespace AJA
 				if (bResult)
 				{
 					std::string FailureReason;
-					if (!Helpers::GetInputVideoFormat(Connection.Card, InTransportType, InChannel, InInputSource, InDesiredInputFormat, VideoFormat, false, FailureReason))
+					if (bInAsGenlock || FoundChannel->bGenlockChannel)
+					{
+						VideoFormat = InDesiredInputFormat;
+					} 
+					else if (!Helpers::GetInputVideoFormat(Connection.Card, InTransportType, InChannel, InInputSource, InDesiredInputFormat, VideoFormat, false, FailureReason))
 					{
 						UE_LOG(LogAjaCore, Error, TEXT("Device: Initialization of the input failed for channel %d on device %S. %S"), uint32_t(InChannel) + 1, Connection.Card->GetDisplayName().c_str(), FailureReason.c_str());
 						bResult = false;
@@ -300,7 +306,7 @@ namespace AJA
 					}
 				}
 
-				if (bResult && FoundChannel->TimecodeFormat != InTimecodeFormat && !bInIsAutoDetected)
+				if (bResult && FoundChannel->TimecodeFormat != InTimecodeFormat && !bInIsAutoDetected && bInAsGenlock)
 				{
 					UE_LOG(LogAjaCore, Warning,  TEXT("Device: Try to enable the channel '%d' on device '%S' with a timecode format that is not the same as the previous one. Timecode may not be decoded properly.\n")
 						, uint32_t(InChannel) + 1, Connection.Card->GetDisplayName().c_str());
@@ -414,6 +420,12 @@ namespace AJA
 			}
 			else
 			{
+				if (Connection.bOutputReferenceSet)
+				{
+					// Connection was also used for genlock but was not initialized as a genlock connection, so mark it as so.
+					FoundChannel->bGenlockChannel = true;
+				}
+
 				if (bInAsOwner)
 				{
 					FoundChannel->bIsOwned = false;
@@ -586,13 +598,13 @@ namespace AJA
 			{
 				if (Connection.OutputReferenceType != InOutputReferenceType)
 				{
-					UE_LOG(LogAjaCore, Error, TEXT("Device: Couldn't set the reference for output on device %S. The type was already setup with %d.\n"), Connection.Card->GetDisplayName().c_str(), int(Connection.OutputReferenceType));
+					UE_LOG(LogAjaCore, Error, TEXT("Device: Couldn't set the reference (%S) for output on device %S. The type was already setup with %S.\n"), Helpers::ReferenceTypeToString(InOutputReferenceType), Connection.Card->GetDisplayName().c_str(), Helpers::ReferenceTypeToString(Connection.OutputReferenceType));
 					return false;
 				}
 
 				if (Connection.OutputReferenceType == EAJAReferenceType::EAJA_REFERENCETYPE_INPUT && Connection.OutputReferenceChannel != InOutputReferenceChannel)
 				{
-					UE_LOG(LogAjaCore, Error, TEXT("Device: Couldn't set the reference for output on device %S. The input was already setup with %d.\n"), Connection.Card->GetDisplayName().c_str(), Connection.OutputReferenceChannel);
+					UE_LOG(LogAjaCore, Error, TEXT("Device: Couldn't set the reference (%S) for output on device %S. The input was already setup with %S.\n"), Helpers::ReferenceTypeToString(InOutputReferenceType), Connection.Card->GetDisplayName().c_str(), Connection.OutputReferenceChannel, Helpers::ReferenceTypeToString(Connection.OutputReferenceType));
 					return false;
 				}
 
@@ -605,7 +617,7 @@ namespace AJA
 			{
 				if (Connection.bAnalogLtcFromReferenceInput)
 				{
-					UE_LOG(LogAjaCore, Error, TEXT("Device: Couldn't set the reference for output channel %d on device %S. The reference is used to read analog LTC.\n"), InOutputReferenceChannel, Connection.Card->GetDisplayName().c_str());
+					UE_LOG(LogAjaCore, Error, TEXT("Device: Couldn't set the reference (%S) for output channel %d on device %S. The reference is used to read analog LTC.\n"), Helpers::ReferenceTypeToString(InOutputReferenceType), InOutputReferenceChannel, Connection.Card->GetDisplayName().c_str());
 					return false;
 				}
 				AJA_CHECK(Connection.Card->SetLTCInputEnable(false));
diff --git a/Engine/Plugins/Media/AjaMedia/Source/Aja/Private/Device.h b/Engine/Plugins/Media/AjaMedia/Source/Aja/Private/Device.h
--- a/Engine/Plugins/Media/AjaMedia/Source/Aja/Private/Device.h
+++ b/Engine/Plugins/Media/AjaMedia/Source/Aja/Private/Device.h
@@ -86,6 +86,7 @@ namespace AJA
 				bool bIsOwned;
 				bool bIsInput;
 				bool bIsAutoDetected;
+				bool bGenlockChannel;
 			};
 
 		public:
@@ -120,9 +121,9 @@ namespace AJA
 				CommandList(DeviceConnection&);
 
 			public:
-				bool RegisterChannel(ETransportType InTransportType, NTV2InputSource InInputSource, NTV2Channel InChannel, bool bInAsInput, bool bConnectChannel, ETimecodeFormat InTimecodeFormat, EPixelFormat InUEPixelFormat, NTV2VideoFormat InDesiredInputFormat, bool bInAsOwner, bool bInIsAutoDetected);
+				bool RegisterChannel(ETransportType InTransportType, NTV2InputSource InInputSource, NTV2Channel InChannel, bool bInAsInput, bool bInAsGenlock, bool bConnectChannel, ETimecodeFormat InTimecodeFormat, EPixelFormat InUEPixelFormat, NTV2VideoFormat InDesiredInputFormat, bool bInAsOwner, bool bInIsAutoDetected);
 				/** RegisterChannel override that handles HDR metadata. */
-				bool RegisterChannel(ETransportType InTransportType, NTV2InputSource InInputSource, NTV2Channel InChannel, bool bInAsInput, bool bConnectChannel, ETimecodeFormat InTimecodeFormat, EPixelFormat InUEPixelFormat, NTV2VideoFormat InDesiredInputFormat, FAjaHDROptions& InOutHDROptions, bool bInAsOwner, bool bInIsAutoDetected);
+				bool RegisterChannel(ETransportType InTransportType, NTV2InputSource InInputSource, NTV2Channel InChannel, bool bInAsInput, bool bInAsGenlock, bool bConnectChannel, ETimecodeFormat InTimecodeFormat, EPixelFormat InUEPixelFormat, NTV2VideoFormat InDesiredInputFormat, FAjaHDROptions& InOutHDROptions, bool bInAsOwner, bool bInIsAutoDetected);
 				bool UnregisterChannel(NTV2Channel InChannel, bool bInAsOwner);
 
 				bool IsChannelConnect(NTV2Channel InChannel);
diff --git a/Engine/Plugins/Media/AjaMedia/Source/Aja/Private/Helpers.cpp b/Engine/Plugins/Media/AjaMedia/Source/Aja/Private/Helpers.cpp
--- a/Engine/Plugins/Media/AjaMedia/Source/Aja/Private/Helpers.cpp
+++ b/Engine/Plugins/Media/AjaMedia/Source/Aja/Private/Helpers.cpp
@@ -1197,6 +1197,17 @@ namespace AJA
 			return "<Invalid>";
 		}
 
+		const char* Helpers::ReferenceTypeToString(EAJAReferenceType InReferenceType)
+		{
+			switch (InReferenceType)
+			{
+			case EAJAReferenceType::EAJA_REFERENCETYPE_EXTERNAL: return "External";
+			case EAJAReferenceType::EAJA_REFERENCETYPE_FREERUN: return "Free Run";
+			case EAJAReferenceType::EAJA_REFERENCETYPE_INPUT: return "Input";
+			}
+			return "<Invalid>";
+		}
+
 		NTV2VideoFormat Helpers::Get372Format(NTV2VideoFormat InFormat)
 		{
 			switch (InFormat)
diff --git a/Engine/Plugins/Media/AjaMedia/Source/Aja/Private/Helpers.h b/Engine/Plugins/Media/AjaMedia/Source/Aja/Private/Helpers.h
--- a/Engine/Plugins/Media/AjaMedia/Source/Aja/Private/Helpers.h
+++ b/Engine/Plugins/Media/AjaMedia/Source/Aja/Private/Helpers.h
@@ -135,6 +135,7 @@ namespace AJA
 			static bool IsSdiTransport(ETransportType InTransportType);
 			static bool IsHdmiTransport(ETransportType InTransportType);
 			static const char* TransportTypeToString(ETransportType InTransportType);
+			static const char* ReferenceTypeToString(EAJAReferenceType InReferenceType);
 
 			static NTV2VideoFormat Get372Format(NTV2VideoFormat InFormat);
 			static NTV2VideoFormat GetLevelA(NTV2VideoFormat InFormat);
diff --git a/Engine/Plugins/Media/AjaMedia/Source/Aja/Private/SyncChannel.cpp b/Engine/Plugins/Media/AjaMedia/Source/Aja/Private/SyncChannel.cpp
--- a/Engine/Plugins/Media/AjaMedia/Source/Aja/Private/SyncChannel.cpp
+++ b/Engine/Plugins/Media/AjaMedia/Source/Aja/Private/SyncChannel.cpp
@@ -254,10 +254,11 @@ namespace AJA
 			}
 
 			const bool bRegisteredAsInput = !SyncOption.bOutput;
-			const bool bConnectChannel = false;
-			const bool bAsOwner = false;
+			constexpr bool bConnectChannel = false;
+			constexpr bool bAsOwner = false;
+			constexpr bool bAsGenlock = true;
 			const EPixelFormat DefaultPixelFormat = EPixelFormat::PF_8BIT_YCBCR;
-			if (!InCommandList.RegisterChannel(SyncOption.TransportType, InputSource, Channel, bRegisteredAsInput, bConnectChannel, SyncOption.TimecodeFormat, DefaultPixelFormat, DesiredVideoFormat, bAsOwner, SyncOption.bAutoDetectFormat))
+			if (!InCommandList.RegisterChannel(SyncOption.TransportType, InputSource, Channel, bRegisteredAsInput, bAsGenlock, bConnectChannel, SyncOption.TimecodeFormat, DefaultPixelFormat, DesiredVideoFormat, bAsOwner, SyncOption.bAutoDetectFormat))
 			{
 				return false;
 			}
diff --git a/Engine/Plugins/Media/AjaMedia/Source/Aja/Private/TimecodeChannel.cpp b/Engine/Plugins/Media/AjaMedia/Source/Aja/Private/TimecodeChannel.cpp
--- a/Engine/Plugins/Media/AjaMedia/Source/Aja/Private/TimecodeChannel.cpp
+++ b/Engine/Plugins/Media/AjaMedia/Source/Aja/Private/TimecodeChannel.cpp
@@ -253,11 +253,13 @@ namespace AJA
 					return false;
 				}
 
-				const bool bRegisteredAsInput = true;
-				const bool bConnectChannel = false;
-				const bool bAsOwner = false;
+				constexpr bool bRegisteredAsInput = true;
+				constexpr bool bConnectChannel = false;
+				constexpr bool bAsOwner = false;
+				constexpr bool bAsGenlock = false;
+
 				const EPixelFormat DefaultPixelFormat = EPixelFormat::PF_8BIT_YCBCR;
-				if (!InCommandList.RegisterChannel(TimecodeOption.TransportType, InputSource, Channel, bRegisteredAsInput, bConnectChannel, TimecodeOption.TimecodeFormat, DefaultPixelFormat, DesiredVideoFormat, bAsOwner, TimecodeOption.bAutoDetectFormat))
+				if (!InCommandList.RegisterChannel(TimecodeOption.TransportType, InputSource, Channel, bRegisteredAsInput, bAsGenlock, bConnectChannel, TimecodeOption.TimecodeFormat, DefaultPixelFormat, DesiredVideoFormat, bAsOwner, TimecodeOption.bAutoDetectFormat))
 				{
 					return false;
 				}
