Reinitialize output channel when connecting output to channel already used for genlock
Do not clear reference type when uninitializing channel
Always use syncthread to genlock AJA

diff --git a/Engine/Plugins/Media/AjaMedia/Source/Aja/Private/ChannelBase.cpp b/Engine/Plugins/Media/AjaMedia/Source/Aja/Private/ChannelBase.cpp
--- a/Engine/Plugins/Media/AjaMedia/Source/Aja/Private/ChannelBase.cpp
+++ b/Engine/Plugins/Media/AjaMedia/Source/Aja/Private/ChannelBase.cpp
@@ -214,7 +214,7 @@ namespace AJA
 
 				if (bRegisteredReference)
 				{
-					InCommandList.UnregisterReference();
+					InCommandList.UnregisterReference(Channel);
 				}
 
 				if (bRegisteredChannel)
diff --git a/Engine/Plugins/Media/AjaMedia/Source/Aja/Private/Device.cpp b/Engine/Plugins/Media/AjaMedia/Source/Aja/Private/Device.cpp
--- a/Engine/Plugins/Media/AjaMedia/Source/Aja/Private/Device.cpp
+++ b/Engine/Plugins/Media/AjaMedia/Source/Aja/Private/Device.cpp
@@ -228,6 +228,11 @@ namespace AJA
 					Info->bIsAutoDetected = bInIsAutoDetected;
 					Info->bGenlockChannel = bInAsGenlock;
 
+					// Spawn a sync thread if there is more than 1 connection that can potentially wait on the same event
+					Info->SyncThread.reset(new SyncThreadInfoThread(&Connection, Info));
+					Info->SyncThread->SetPriority(AJA_ThreadPriority_High);
+					Info->SyncThread->Start();
+
 					{
 						AJAAutoLock AutoLock(&Connection.ChannelInfoLock);
 						ConnectionChannelInfos.push_back(Info);
@@ -238,9 +243,10 @@ namespace AJA
 			{
 				ChannelInfo* FoundChannel = *FoundIterator;
 
-				if (bReinitialize)
+				// If channel was used for Genlock, reinitialize channel since routing wasn't done.
+				if (bReinitialize || (FoundChannel->bGenlockChannel && !bInAsGenlock))
 				{
-					if (!FoundChannel->bIsAutoDetected)
+					if (!FoundChannel->bIsAutoDetected && !(FoundChannel->bGenlockChannel && !bInAsGenlock))
 					{
 						UE_LOG(LogAjaCore, Warning,  TEXT("Device: Autodetect reinitialized channel '%d' on device '%S' but the channel was not in autodetect mode.\n")
 							, uint32_t(InChannel) + 1, Connection.Card->GetDisplayName().c_str());
@@ -320,14 +326,6 @@ namespace AJA
 					{
 						FoundChannel->bIsOwned = bInAsOwner;
 					}
-
-					// Spawn a sync thread if there is more than 1 connection that can potentially wait on the same event
-					if (!FoundChannel->SyncThread)
-					{
-						FoundChannel->SyncThread.reset(new SyncThreadInfoThread(&Connection, FoundChannel));
-						FoundChannel->SyncThread->SetPriority(AJA_ThreadPriority_High);
-						FoundChannel->SyncThread->Start();
-					}
 				}
 			}
 
@@ -649,11 +647,23 @@ namespace AJA
 			return true;
 		}
 
-		void DeviceConnection::CommandList::UnregisterReference()
+		void DeviceConnection::CommandList::UnregisterReference(NTV2Channel InChannel)
 		{
 			AJAAutoLock Lock(&Connection.ChannelLock);
+			auto Found = std::find_if(std::begin(Connection.ChannelInfos), std::end(Connection.ChannelInfos), [=](const ChannelInfo* Info) { return Info->Channel == InChannel; });
+			
+			if (Found == std::end(Connection.ChannelInfos))
+			{
+				Connection.bOutputReferenceSet = false;
+				return;
+			}
 
-			Connection.bOutputReferenceSet = false;
+			const ChannelInfo* FoundChannel = *Found;
+			const bool bOutputChannelUsingRefPin = FoundChannel->RefCounter > 1 && Connection.OutputReferenceType == AJA::EAJAReferenceType::EAJA_REFERENCETYPE_EXTERNAL;
+			if (!bOutputChannelUsingRefPin)
+			{
+				Connection.bOutputReferenceSet = false;
+			}
 		}
 
 		/* Device implementation
@@ -755,7 +765,8 @@ namespace AJA
 				}
 				else
 				{
-					bResult = SyncThreadInfoThread::WaitForVerticalInterrupt(this, FoundChannel);
+					bResult = false;
+					ensureAlwaysMsgf(false, TEXT("Expected a sync thread but there was none."));
 				}
 
 				bool bDoAnotherWaitForField = false;
@@ -807,7 +818,8 @@ namespace AJA
 				}
 				else
 				{
-					bResult = SyncThreadInfoThread::WaitForVerticalInterrupt(this, FoundChannel);
+					bResult = false;
+					ensureAlwaysMsgf(false, TEXT("Expected a sync thread but there was none."));
 				}
 
 				if (bResult)
diff --git a/Engine/Plugins/Media/AjaMedia/Source/Aja/Private/Device.h b/Engine/Plugins/Media/AjaMedia/Source/Aja/Private/Device.h
--- a/Engine/Plugins/Media/AjaMedia/Source/Aja/Private/Device.h
+++ b/Engine/Plugins/Media/AjaMedia/Source/Aja/Private/Device.h
@@ -52,7 +52,7 @@ namespace AJA
 				volatile bool bWaitResult;
 
 			public:
-				static bool WaitForVerticalInterrupt(DeviceConnection* InDeviceConnection, ChannelInfo* InChannelInfo);
+				bool WaitForVerticalInterrupt(DeviceConnection* InDeviceConnection, ChannelInfo* InChannelInfo);
 				static bool IsCurrentField(DeviceConnection* InDeviceConnection, ChannelInfo* InChannelInfo, NTV2FieldID InFieldId = NTV2_FIELD0);
 				bool Wait_ExternalThread();
 				virtual bool ThreadLoop() override;
@@ -133,7 +133,7 @@ namespace AJA
 				void UnregisterAnalogLtc(bool bUseReferencePin);
 
 				bool RegisterReference(EAJAReferenceType OutputReferenceType, NTV2Channel InChannel);
-				void UnregisterReference();
+				void UnregisterReference(NTV2Channel InChannel);
 
 			private:
 				DeviceConnection& Connection;
diff --git a/Engine/Plugins/Media/AjaMedia/Source/Aja/Private/SyncChannel.cpp b/Engine/Plugins/Media/AjaMedia/Source/Aja/Private/SyncChannel.cpp
--- a/Engine/Plugins/Media/AjaMedia/Source/Aja/Private/SyncChannel.cpp
+++ b/Engine/Plugins/Media/AjaMedia/Source/Aja/Private/SyncChannel.cpp
@@ -180,7 +180,7 @@ namespace AJA
 			{
 				if (bRegisterReference)
 				{
-					InCommandList.UnregisterReference();
+					InCommandList.UnregisterReference(Channel);
 				}
 
 				if (bRegisterChannel)
