Fix Aja Autodetection putting card in uniformat mode.

diff --git a/Engine/Plugins/Media/AjaMedia/Source/Aja/Private/AutoDetectChannel.cpp b/Engine/Plugins/Media/AjaMedia/Source/Aja/Private/AutoDetectChannel.cpp
--- a/Engine/Plugins/Media/AjaMedia/Source/Aja/Private/AutoDetectChannel.cpp
+++ b/Engine/Plugins/Media/AjaMedia/Source/Aja/Private/AutoDetectChannel.cpp
@@ -20,7 +20,6 @@ namespace AJA
 			{
 				// Get the first device. It doesn't matter which one since we want only want the device to detect the input for other device.
 				AJADeviceOptions DeviceOption = AJADeviceOptions(0);
-				DeviceOption.bWantMultiFormatMode = false;
 
 				Device = DeviceCache::GetDevice(DeviceOption);
 				std::shared_ptr<AutoDetectChannelInitialize_DeviceCommand> SharedCommand(new AutoDetectChannelInitialize_DeviceCommand(Device, AutoDetectImplementation));
