Aja - Allow initializing LTC Channel with same Framerate family without multi format mode.

diff --git a/Engine/Plugins/Media/AjaMedia/Source/Aja/Private/Device.cpp b/Engine/Plugins/Media/AjaMedia/Source/Aja/Private/Device.cpp
--- a/Engine/Plugins/Media/AjaMedia/Source/Aja/Private/Device.cpp
+++ b/Engine/Plugins/Media/AjaMedia/Source/Aja/Private/Device.cpp
@@ -524,17 +524,9 @@ namespace AJA
 					//LTC reading relies on card/Channel 1 frame rate to be compatible
 					const ChannelInfo* FoundChannel = *Found;
 					const NTV2FrameRate ChannelFrameRate = ::GetNTV2FrameRateFromVideoFormat(FoundChannel->VideoFormat);
-					if (::NTV2DeviceCanDoMultiFormat(Connection.Card->GetDeviceID()))
+					if (!::IsMultiFormatCompatible(InFrameRate, ChannelFrameRate))
 					{
-						if (!::IsMultiFormatCompatible(InFrameRate, ChannelFrameRate))
-						{
-							UE_LOG(LogAjaCore, Error, TEXT("Device: Trying to read LTC with FrameRate '%S' but it's not compatible with card's (channel 1) current FrameRate of '%S' on device '%S'.\n"), NTV2FrameRateToString(InFrameRate).c_str(), NTV2FrameRateToString(ChannelFrameRate).c_str(), Connection.Card->GetDisplayName().c_str());
-							bResult = false;
-						}
-					}
-					else if (InFrameRate != ChannelFrameRate)
-					{
-						UE_LOG(LogAjaCore, Error, TEXT("Device: Trying to read LTC with FrameRate '%S' but card's (channel 1) current FrameRate is '%S' and device '%S' can't do multi-format.\n"), NTV2FrameRateToString(InFrameRate).c_str(), NTV2FrameRateToString(ChannelFrameRate).c_str(), Connection.Card->GetDisplayName().c_str());
+						UE_LOG(LogAjaCore, Error, TEXT("Device: Trying to read LTC with FrameRate '%S' but it's not compatible with card's (channel 1) current FrameRate of '%S' on device '%S'.\n"), NTV2FrameRateToString(InFrameRate).c_str(), NTV2FrameRateToString(ChannelFrameRate).c_str(), Connection.Card->GetDisplayName().c_str());
 						bResult = false;
 					}
 				}
