Aja - Fix GPU Direct not working with Aja Key&Fill
Fix Multiformat 1080p/4k not working, was caused by overlap between the frame buffers.
We now index the frames using the biggest gaps possible in case 4K frames are used.

diff --git a/Engine/Plugins/Media/AjaMedia/Source/Aja/Private/Device.cpp b/Engine/Plugins/Media/AjaMedia/Source/Aja/Private/Device.cpp
--- a/Engine/Plugins/Media/AjaMedia/Source/Aja/Private/Device.cpp
+++ b/Engine/Plugins/Media/AjaMedia/Source/Aja/Private/Device.cpp
@@ -330,9 +330,8 @@ namespace AJA
 			}
 
 			//Verify if Channel1 is used. If not, we need to be sure the card's main format is configured
-			if (bResult)
+			if (bResult && !::NTV2DeviceCanDoMultiFormat(Connection.Card->GetDeviceID()))
 			{
-				const NTV2FrameRate DesiredFrameRate = GetNTV2FrameRateFromVideoFormat(VideoFormat);
 				auto Found = std::find_if(std::begin(Connection.ChannelInfos), std::end(Connection.ChannelInfos), [=](const ChannelInfo* Info) { return Info->Channel == NTV2_CHANNEL1; });
 				if (Found == std::end(Connection.ChannelInfos))
 				{
@@ -1302,19 +1301,7 @@ namespace AJA
 		uint32_t DeviceConnection::Lock_AcquireBaseFrameIndex(NTV2Channel InChannel, NTV2VideoFormat InVideoFormat) const
 		{
 			const uint32_t MaxValue = ::NTV2DeviceGetNumberFrameBuffers(Card->GetDeviceID());
-
-			const bool bIsQuadFrameFormat = NTV2_IS_4K_VIDEO_FORMAT(InVideoFormat);
-
-			uint32_t BaseFrameIndex = 0;
-
-			if (bIsQuadFrameFormat)
-			{
-				BaseFrameIndex = InChannel * NumberOfFrameToAquire;
-			}
-			else
-			{
-				BaseFrameIndex = InChannel * NumberOfFrameToAquire * 4;
-			}
+			const uint32_t BaseFrameIndex = InChannel * NumberOfFrameToAquire * 4;
 
 			if (BaseFrameIndex < MaxValue)
 			{
diff --git a/Engine/Plugins/Media/AjaMedia/Source/Aja/Private/OutputChannel.cpp b/Engine/Plugins/Media/AjaMedia/Source/Aja/Private/OutputChannel.cpp
--- a/Engine/Plugins/Media/AjaMedia/Source/Aja/Private/OutputChannel.cpp
+++ b/Engine/Plugins/Media/AjaMedia/Source/Aja/Private/OutputChannel.cpp
@@ -684,7 +684,14 @@ namespace AJA
 			 
 			// Todo jroy: Verify if Width needs to be fixed, by using the GPUTextureTransfer register struct directly, same for pixel format.
 			const ULWord Stride = FormatDescriptor.GetBytesPerRow();
-			const ULWord Width = GetOptions().PixelFormat == EPixelFormat::PF_8BIT_YCBCR ? Stride / 4 : Stride / 16;
+			ULWord Width = Stride / 4;
+
+			if (GetOptions().PixelFormat == EPixelFormat::PF_10BIT_YCBCR)
+			{
+				// YUV210 has a different stride.
+				Width /= 4;
+			}
+
 			const ULWord Height = FormatDescriptor.GetRasterHeight();
 
 			if (!bDMABuffersRegistered)
@@ -696,7 +703,18 @@ namespace AJA
 					Args.Width = Width;
 					Args.Stride = Stride;
 					Args.Buffer = IttFrame->VideoBuffer;
-					Args.PixelFormat = GetOptions().PixelFormat == EPixelFormat::PF_8BIT_YCBCR ? UE::GPUTextureTransfer::EPixelFormat::PF_8Bit : UE::GPUTextureTransfer::EPixelFormat::PF_10Bit;
+
+					if (GetOptions().PixelFormat == EPixelFormat::PF_8BIT_YCBCR || GetOptions().PixelFormat == EPixelFormat::PF_8BIT_ARGB
+						|| GetOptions().PixelFormat == EPixelFormat::PF_10BIT_RGB) // 10 Bit RGB should be considered as 8 bit as far as GPUDirect is concerned.
+					{
+
+						Args.PixelFormat = UE::GPUTextureTransfer::EPixelFormat::PF_8Bit;
+					}
+					else
+					{
+						Args.PixelFormat = UE::GPUTextureTransfer::EPixelFormat::PF_10Bit;
+					}
+
 					TextureTransfer->RegisterBuffer(Args);
 				}
 
@@ -1053,7 +1071,16 @@ namespace AJA
 				}
 				if (UseVideo())
 				{
-					AJA_CHECK(GetDevice().DMAWriteFrame(BaseFrameIndex + Index, reinterpret_cast<ULWord*>(CurrentFrame->VideoBuffer), CurrentFrame->CopiedVideoBufferSize));
+					uint32 FrameIndex = BaseFrameIndex + Index;
+
+					if (NTV2_IS_4K_VIDEO_FORMAT(GetOptions().VideoFormatIndex))
+					{
+						// 4K Format uses a different indexing.
+						FrameIndex *= 4;
+					}
+					
+					AJA_CHECK(GetDevice().DMAWriteFrame(FrameIndex, reinterpret_cast<ULWord*>(CurrentFrame->VideoBuffer), CurrentFrame->CopiedVideoBufferSize));
+
 					if (TextureTransfer && GetOptions().bUseGPUDMA)
 					{
 						TextureTransfer->EndSync(CurrentFrame->VideoBuffer);
@@ -1184,7 +1211,16 @@ namespace AJA
 							const UE::GPUTextureTransfer::ETransferDirection Direction = UE::GPUTextureTransfer::ETransferDirection::GPU_TO_CPU;
 							TextureTransfer->BeginSync(AvailableReadingFrame->VideoBuffer, Direction);
 						}
-						bRunning = bRunning && GetDevice().DMAWriteFrame(BaseFrameIndex + CurrentOutFrame, reinterpret_cast<ULWord*>(AvailableReadingFrame->VideoBuffer), AvailableReadingFrame->CopiedVideoBufferSize);
+
+						uint32 FrameIndex = BaseFrameIndex + CurrentOutFrame;
+						if (NTV2_IS_4K_VIDEO_FORMAT(GetOptions().VideoFormatIndex))
+						{
+							// 4K Format uses a different indexing.
+							FrameIndex *= 4;
+						}
+
+						bRunning = bRunning && GetDevice().DMAWriteFrame(FrameIndex, reinterpret_cast<ULWord*>(AvailableReadingFrame->VideoBuffer), AvailableReadingFrame->CopiedVideoBufferSize);
+
 						if (TextureTransfer && GetOptions().bUseGPUDMA)
 						{
 							TextureTransfer->EndSync(AvailableReadingFrame->VideoBuffer);
diff --git a/Engine/Plugins/Media/AjaMedia/Source/AjaMediaOutput/Private/AjaMediaCapture.cpp b/Engine/Plugins/Media/AjaMedia/Source/AjaMediaOutput/Private/AjaMediaCapture.cpp
--- a/Engine/Plugins/Media/AjaMedia/Source/AjaMediaOutput/Private/AjaMediaCapture.cpp
+++ b/Engine/Plugins/Media/AjaMedia/Source/AjaMediaOutput/Private/AjaMediaCapture.cpp
@@ -625,6 +625,12 @@ void UAjaMediaCapture::LockDMATexture_RenderThread(FTextureRHIRef InTexture)
 				Args.PixelFormat = UE::GPUTextureTransfer::EPixelFormat::PF_10Bit;
 				Args.Stride = Args.Width * 16;
 			}
+			else if (Args.RHITexture->GetFormat() == EPixelFormat::PF_A2B10G10R10)
+			{
+				// RGB 10 bit can be considered as 8 PF_8 bits by GPUDirect.
+				Args.PixelFormat = UE::GPUTextureTransfer::EPixelFormat::PF_8Bit;
+				Args.Stride = Args.Width * 4;
+			}
 			else
 			{
 				checkf(false, TEXT("Format not supported"));
