- Fixing crash on shutdown when doing aja media capture

diff --git a/Engine/Plugins/Media/AjaMedia/Source/AjaMediaOutput/Private/AjaMediaCapture.cpp b/Engine/Plugins/Media/AjaMedia/Source/AjaMediaOutput/Private/AjaMediaCapture.cpp
--- a/Engine/Plugins/Media/AjaMedia/Source/AjaMediaOutput/Private/AjaMediaCapture.cpp
+++ b/Engine/Plugins/Media/AjaMedia/Source/AjaMediaOutput/Private/AjaMediaCapture.cpp
@@ -360,7 +360,11 @@ void UAjaMediaCapture::StopCaptureImpl(bool bAllowPendingFrameToBeProcess)
 
 		if (GEngine)
 		{
-			GEngine->GetEngineSubsystem<UMediaIOCoreSubsystem>()->OnBufferReceived_AudioThread().RemoveAll(this);
+			UMediaIOCoreSubsystem* SubSystem = GEngine->GetEngineSubsystem<UMediaIOCoreSubsystem>();
+			if (SubSystem)
+			{
+				SubSystem->OnBufferReceived_AudioThread().RemoveAll(this);
+			}
 		}
 
 		AudioOutput.Reset();
