Aja - Fix 4k Output writing into the wrong buffers

Fix NTV2 device not correctly setting card framerate when starting output on other channels  than 1

diff --git a/Engine/Plugins/Media/AjaMedia/Source/Aja/Private/Device.cpp b/Engine/Plugins/Media/AjaMedia/Source/Aja/Private/Device.cpp
--- a/Engine/Plugins/Media/AjaMedia/Source/Aja/Private/Device.cpp
+++ b/Engine/Plugins/Media/AjaMedia/Source/Aja/Private/Device.cpp
@@ -1027,6 +1027,7 @@ namespace AJA
 			{
 				bool bIsMultiFormatEnabled = false;
 				AJA_CHECK(InCard->GetMultiFormatMode(bIsMultiFormatEnabled));
+
 				if (bIsMultiFormatEnabled)
 				{
 					if (!::IsMultiFormatCompatible(NotMultiVideoFormat, OutFoundVideoFormat))
@@ -1041,6 +1042,42 @@ namespace AJA
 					return false;
 				}
 			}
+
+			if (InChannel != NTV2_CHANNEL1)
+			{
+				// Channel 1 dictates the card's framerate, so make sure we are running at a compatible framerate.
+				NTV2FrameRate DeviceFramerate = NTV2_FRAMERATE_INVALID;
+				AJA_CHECK(InCard->GetFrameRate(DeviceFramerate));
+
+				NTV2FrameRate NewVideoFormatFrameRate = GetNTV2FrameRateFromVideoFormat(OutFoundVideoFormat);
+
+				NTV2FrameRate NewVideoFormatFrameRateFamily = GetFrameRateFamily(NewVideoFormatFrameRate);
+
+				if (GetFrameRateFamily(DeviceFramerate) != NewVideoFormatFrameRateFamily)
+				{
+					// Device framerate doesn't match... We'll override it if the channel isn't in use.
+					AJAAutoLock AutoLock(&ChannelInfoLock);
+
+					// Check all channels that are already in use before changing the framerate.
+					// We can't change the card framerate if it's driven by a different channel.
+					for (ChannelInfo* Info : ChannelInfos)
+					{
+						if (Info)
+						{
+							NTV2FrameRate ChannelVideoFormatFrameRate = GetNTV2FrameRateFromVideoFormat(Info->VideoFormat);
+							if (NewVideoFormatFrameRateFamily != GetFrameRateFamily(ChannelVideoFormatFrameRate))
+							{
+								UE_LOG(LogAjaCore, Error, TEXT("Device: Can't enable channel %d device '%S' because its framerate doesn't match an existing channel's framerate family."), uint32_t(InChannel) + 1, InCard->GetDisplayName().c_str());
+								return false;
+							}
+						}
+					}
+
+					AJA_CHECK(InCard->SetFrameRate(NewVideoFormatFrameRate));
+				}
+			}
+
+
 			NotMultiVideoFormat = OutFoundVideoFormat;
 
 			const bool bIsSDI = Helpers::IsSdiTransport(InTransportType);
diff --git a/Engine/Plugins/Media/AjaMedia/Source/Aja/Private/OutputChannel.cpp b/Engine/Plugins/Media/AjaMedia/Source/Aja/Private/OutputChannel.cpp
--- a/Engine/Plugins/Media/AjaMedia/Source/Aja/Private/OutputChannel.cpp
+++ b/Engine/Plugins/Media/AjaMedia/Source/Aja/Private/OutputChannel.cpp
@@ -1070,15 +1070,7 @@ namespace AJA
 				}
 				if (UseVideo())
 				{
-					uint32 FrameIndex = BaseFrameIndex + Index;
-
-					if (NTV2_IS_4K_VIDEO_FORMAT(GetOptions().VideoFormatIndex))
-					{
-						// 4K Format uses a different indexing.
-						FrameIndex *= 4;
-					}
-					
-					AJA_CHECK(GetDevice().DMAWriteFrame(FrameIndex, reinterpret_cast<ULWord*>(CurrentFrame->VideoBuffer), CurrentFrame->CopiedVideoBufferSize));
+					AJA_CHECK(GetDevice().DMAWriteFrame(BaseFrameIndex + Index, reinterpret_cast<ULWord*>(CurrentFrame->VideoBuffer), CurrentFrame->CopiedVideoBufferSize, Channel));
 
 					if (TextureTransfer && GetOptions().bUseGPUDMA)
 					{
@@ -1211,14 +1203,7 @@ namespace AJA
 							TextureTransfer->BeginSync(AvailableReadingFrame->VideoBuffer, Direction);
 						}
 
-						uint32 FrameIndex = BaseFrameIndex + CurrentOutFrame;
-						if (NTV2_IS_4K_VIDEO_FORMAT(GetOptions().VideoFormatIndex))
-						{
-							// 4K Format uses a different indexing.
-							FrameIndex *= 4;
-						}
-
-						bRunning = bRunning && GetDevice().DMAWriteFrame(FrameIndex, reinterpret_cast<ULWord*>(AvailableReadingFrame->VideoBuffer), AvailableReadingFrame->CopiedVideoBufferSize);
+						bRunning = bRunning && GetDevice().DMAWriteFrame(BaseFrameIndex + CurrentOutFrame, reinterpret_cast<ULWord*>(AvailableReadingFrame->VideoBuffer), AvailableReadingFrame->CopiedVideoBufferSize, Channel);
 
 						if (TextureTransfer && GetOptions().bUseGPUDMA)
 						{
