Add option to output interlace video as progressive. ( Experimental )

Could help in cases where you have an interlace signal going through UE to keep frames as-is. Combining OutputInterlaceAsProgressive with using Deinterlacer - None on an interlace input should let the interlace frame go through Unreal untouched.

diff --git a/Engine/Plugins/Media/AjaMedia/Source/Aja/Private/OutputChannel.cpp b/Engine/Plugins/Media/AjaMedia/Source/Aja/Private/OutputChannel.cpp
--- a/Engine/Plugins/Media/AjaMedia/Source/Aja/Private/OutputChannel.cpp
+++ b/Engine/Plugins/Media/AjaMedia/Source/Aja/Private/OutputChannel.cpp
@@ -240,7 +240,7 @@ namespace AJA
 		// For Interlaced, there is 2 options. Use Timecode to identify which is the odd and even field.
 		OutputChannelThread::Frame* OutputChannelThread::FetchAvailableWritingFrame(const AJAOutputFrameBufferData& InFrameData)
 		{
-			const bool bIsProgressive = ::IsProgressivePicture(VideoFormat);
+			const bool bIsProgressive = ::IsProgressivePicture(VideoFormat) || Options.bOutputInterlaceAsProgressive;
 
 			Frame* AvailableWritingFrame = nullptr;
 			{
@@ -483,7 +483,7 @@ namespace AJA
 
 		bool OutputChannelThread::SetAncillaryFrameData(const AJAOutputFrameBufferData& InFrameData, uint8_t* InAncillaryBuffer, uint32_t InAncillaryBufferSize)
 		{
-			const bool bIsProgressive = ::IsProgressivePicture(VideoFormat);
+			const bool bIsProgressive = ::IsProgressivePicture(VideoFormat) || Options.bOutputInterlaceAsProgressive;
 
 			if (!UseAncillary())
 			{
@@ -578,7 +578,7 @@ namespace AJA
 
 		bool OutputChannelThread::SetVideoFrameData(const AJAOutputFrameBufferData& InFrameData, uint8_t* InVideoBuffer, uint32_t InVideoBufferSize)
 		{
-			const bool bIsProgressive = ::IsProgressivePicture(VideoFormat);
+			const bool bIsProgressive = ::IsProgressivePicture(VideoFormat) || Options.bOutputInterlaceAsProgressive;
 
 			if (!UseVideo())
 			{
@@ -650,7 +650,7 @@ namespace AJA
 
 		bool OutputChannelThread::SetVideoFrameData(const AJAOutputFrameBufferData& InFrameData, FRHITexture* RHITexture)
 		{
-			const bool bIsProgressive = ::IsProgressivePicture(VideoFormat);
+			const bool bIsProgressive = ::IsProgressivePicture(VideoFormat) || Options.bOutputInterlaceAsProgressive;
 
 			if (!bIsProgressive)
 			{
@@ -869,7 +869,7 @@ namespace AJA
 
 		void OutputChannelThread::Thread_AutoCirculateLoop()
 		{
-			const bool bIsProgressive = ::IsProgressivePicture(VideoFormat);
+			const bool bIsProgressive = ::IsProgressivePicture(VideoFormat) || Options.bOutputInterlaceAsProgressive;
 			bool bRunning = true;
 			bool bSleep = true;
 
@@ -1118,7 +1118,7 @@ namespace AJA
 		void OutputChannelThread::Thread_PingPongLoop()
 		{
 			bool bHaveOutputOnce = false;
-			const bool bIsProgressive = ::IsProgressivePicture(VideoFormat);
+			const bool bIsProgressive = ::IsProgressivePicture(VideoFormat) || Options.bOutputInterlaceAsProgressive;
 
 			uint32_t CurrentOutFrame = 0;
 			CurrentOutFrame ^= 1;
diff --git a/Engine/Plugins/Media/AjaMedia/Source/Aja/Public/AJALib.h b/Engine/Plugins/Media/AjaMedia/Source/Aja/Public/AJALib.h
--- a/Engine/Plugins/Media/AjaMedia/Source/Aja/Public/AJALib.h
+++ b/Engine/Plugins/Media/AjaMedia/Source/Aja/Public/AJALib.h
@@ -459,6 +459,7 @@ namespace AJA
 				uint32_t bOutput : 1; // port is output
 				uint32_t bUseKey : 1; // output will also sent the key on OutputKeyPortIndex
 				uint32_t bOutputInterlacedFieldsTimecodeNeedToMatch : 1; // when trying to find the odd field that correspond to the even field, the 2 timecode need to match
+				uint32_t bOutputInterlaceAsProgressive : 1; // whether to treat outgoing interlace frames as progressive.
 				uint32_t bUseAncillary : 1; // enable ANC system
 				uint32_t bUseAudio : 1; // enable audio input/output
 				uint32_t bUseVideo : 1; // enable video input/output
diff --git a/Engine/Plugins/Media/AjaMedia/Source/AjaMediaOutput/Private/AjaMediaCapture.cpp b/Engine/Plugins/Media/AjaMedia/Source/AjaMediaOutput/Private/AjaMediaCapture.cpp
--- a/Engine/Plugins/Media/AjaMedia/Source/AjaMediaOutput/Private/AjaMediaCapture.cpp
+++ b/Engine/Plugins/Media/AjaMedia/Source/AjaMediaOutput/Private/AjaMediaCapture.cpp
@@ -508,6 +508,7 @@ bool UAjaMediaCapture::InitAJA(UAjaMediaOutput* InAjaMediaOutput)
 	ChannelOptions.bUseAudio = InAjaMediaOutput->bOutputAudio;
 	ChannelOptions.bUseVideo = true;
 	ChannelOptions.bOutputInterlacedFieldsTimecodeNeedToMatch = InAjaMediaOutput->bInterlacedFieldsTimecodeNeedToMatch && Descriptor.bIsInterlacedStandard && InAjaMediaOutput->TimecodeFormat != EMediaIOTimecodeFormat::None;
+	ChannelOptions.bOutputInterlaceAsProgressive = InAjaMediaOutput->bOutputInterlaceAsProgressive;
 	ChannelOptions.bDisplayWarningIfDropFrames = bLogDropFrame;
 	ChannelOptions.bConvertOutputLevelAToB = InAjaMediaOutput->bOutputIn3GLevelB && Descriptor.bIsVideoFormatA;
 	ChannelOptions.TransportType = AjaMediaCaptureUtils::ConvertTransportType(InAjaMediaOutput->OutputConfiguration.MediaConfiguration.MediaConnection.TransportType, InAjaMediaOutput->OutputConfiguration.MediaConfiguration.MediaConnection.QuadTransportType);
diff --git a/Engine/Plugins/Media/AjaMedia/Source/AjaMediaOutput/Public/AjaMediaOutput.h b/Engine/Plugins/Media/AjaMedia/Source/AjaMediaOutput/Public/AjaMediaOutput.h
--- a/Engine/Plugins/Media/AjaMedia/Source/AjaMediaOutput/Public/AjaMediaOutput.h
+++ b/Engine/Plugins/Media/AjaMedia/Source/AjaMediaOutput/Public/AjaMediaOutput.h
@@ -112,6 +112,12 @@ public:
 	UPROPERTY(EditAnywhere, AdvancedDisplay, Category="Output")
 	bool bInterlacedFieldsTimecodeNeedToMatch;
 
+	/**
+	 * If this is enabled when outputting media in interlace format, frames will be treated as progressive and will be outputted as-is. (Experimental)
+	 */
+	UPROPERTY(EditAnywhere, AdvancedDisplay, Category = "Output")
+	bool bOutputInterlaceAsProgressive = false;
+
 	/**
 	 * HDR Metadata of the video signal.
 	 */
diff --git a/Engine/Plugins/Media/BlackmagicMedia/Source/BlackmagicCore/Private/BlackmagicLib.cpp b/Engine/Plugins/Media/BlackmagicMedia/Source/BlackmagicCore/Private/BlackmagicLib.cpp
--- a/Engine/Plugins/Media/BlackmagicMedia/Source/BlackmagicCore/Private/BlackmagicLib.cpp
+++ b/Engine/Plugins/Media/BlackmagicMedia/Source/BlackmagicCore/Private/BlackmagicLib.cpp
@@ -110,6 +110,7 @@ namespace BlackmagicDesign
 		, bOutputVideo(true)
 		, bOutputAudio(false)
 		, bOutputInterlacedFieldsTimecodeNeedToMatch(false)
+		, bOutputInterlaceAsProgressive(false)
 		, bLogDropFrames(false)
 		, bUseGPUDMA(false)
 		, bScheduleInDifferentThread(false)
diff --git a/Engine/Plugins/Media/BlackmagicMedia/Source/BlackmagicCore/Private/BlackmagicOutputChannel.cpp b/Engine/Plugins/Media/BlackmagicMedia/Source/BlackmagicCore/Private/BlackmagicOutputChannel.cpp
--- a/Engine/Plugins/Media/BlackmagicMedia/Source/BlackmagicCore/Private/BlackmagicOutputChannel.cpp
+++ b/Engine/Plugins/Media/BlackmagicMedia/Source/BlackmagicCore/Private/BlackmagicOutputChannel.cpp
@@ -1222,7 +1222,7 @@ namespace BlackmagicDesign
 
 		bool FOutputChannel::SetVideoFrameData(const FFrameDescriptor& InFrame)
 		{
-			const bool bIsProgressive = ChannelOptions.FormatInfo.FieldDominance != EFieldDominance::Interlaced;
+			const bool bIsProgressive = ChannelOptions.FormatInfo.FieldDominance != EFieldDominance::Interlaced || ChannelOptions.bOutputInterlaceAsProgressive;
 
 			if (InFrame.VideoBuffer == nullptr)
 			{
@@ -1258,7 +1258,7 @@ namespace BlackmagicDesign
 			}
 			else
 			{
-				// only write the even or odd frame
+				// Only write the even or odd frame
 				for (int32_t IndexY = bField1 ? 0 : 1; IndexY < InFrame.VideoHeight; IndexY += 2)
 				{
 					int8_t* Destination = reinterpret_cast<int8_t*>(DestinationBuffer) + (Stride*IndexY);
@@ -1295,7 +1295,7 @@ namespace BlackmagicDesign
 		bool FOutputChannel::SetVideoFrameData(FFrameDescriptor_GPUDMA& InFrame)
 		{
 #if PLATFORM_WINDOWS
-			const bool bIsProgressive = ChannelOptions.FormatInfo.FieldDominance != EFieldDominance::Interlaced;
+			const bool bIsProgressive = ChannelOptions.FormatInfo.FieldDominance != EFieldDominance::Interlaced || ChannelOptions.bOutputInterlaceAsProgressive;
 
 			if (!bIsProgressive)
 			{
@@ -1407,7 +1407,7 @@ namespace BlackmagicDesign
 		// For Interlaced, there is 2 options. Use Timecode to identify which is the odd and even field. Use FrameIdentifier which is even and odd field.
 		BlackmagicDesign::Private::FOutputChannel::FOutputFrame* FOutputChannel::FetchAvailableWritingFrame(const FBaseFrameData& InFrame)
 		{
-			const bool bIsProgressive = ChannelOptions.FormatInfo.FieldDominance != EFieldDominance::Interlaced;
+			const bool bIsProgressive = ChannelOptions.FormatInfo.FieldDominance != EFieldDominance::Interlaced || ChannelOptions.bOutputInterlaceAsProgressive;
 			return bIsProgressive ? FetchAvailableWritingFrameProgressive(InFrame) : FetchAvailableWritingFrameInterlaced(InFrame);
 		}
 
diff --git a/Engine/Plugins/Media/BlackmagicMedia/Source/BlackmagicCore/Public/BlackmagicLib.h b/Engine/Plugins/Media/BlackmagicMedia/Source/BlackmagicCore/Public/BlackmagicLib.h
--- a/Engine/Plugins/Media/BlackmagicMedia/Source/BlackmagicCore/Public/BlackmagicLib.h
+++ b/Engine/Plugins/Media/BlackmagicMedia/Source/BlackmagicCore/Public/BlackmagicLib.h
@@ -263,6 +263,7 @@ namespace BlackmagicDesign
 		bool bOutputVideo;
 		bool bOutputAudio;
 		bool bOutputInterlacedFieldsTimecodeNeedToMatch;
+		bool bOutputInterlaceAsProgressive;
 		bool bLogDropFrames;
 		bool bUseGPUDMA;
 		bool bScheduleInDifferentThread;
diff --git a/Engine/Plugins/Media/BlackmagicMedia/Source/BlackmagicMediaOutput/Private/BlackmagicMediaCapture.cpp b/Engine/Plugins/Media/BlackmagicMedia/Source/BlackmagicMediaOutput/Private/BlackmagicMediaCapture.cpp
--- a/Engine/Plugins/Media/BlackmagicMedia/Source/BlackmagicMediaOutput/Private/BlackmagicMediaCapture.cpp
+++ b/Engine/Plugins/Media/BlackmagicMedia/Source/BlackmagicMediaOutput/Private/BlackmagicMediaCapture.cpp
@@ -679,6 +679,7 @@ bool UBlackmagicMediaCapture::InitBlackmagic(UBlackmagicMediaOutput* InBlackmagi
 	ChannelOptions.NumberOfBuffers = FMath::Clamp(InBlackmagicMediaOutput->NumberOfBlackmagicBuffers, 3, 4);
 	ChannelOptions.bOutputVideo = true;
 	ChannelOptions.bOutputInterlacedFieldsTimecodeNeedToMatch = InBlackmagicMediaOutput->bInterlacedFieldsTimecodeNeedToMatch && InBlackmagicMediaOutput->OutputConfiguration.MediaConfiguration.MediaMode.Standard == EMediaIOStandardType::Interlaced && InBlackmagicMediaOutput->TimecodeFormat != EMediaIOTimecodeFormat::None;
+	ChannelOptions.bOutputInterlaceAsProgressive = InBlackmagicMediaOutput->bOutputInterlaceAsProgressive;
 	ChannelOptions.bLogDropFrames = bLogDropFrame;
 	ChannelOptions.bUseGPUDMA = ShouldCaptureRHIResource();
 	ChannelOptions.bScheduleInDifferentThread = InBlackmagicMediaOutput->bUseMultithreadedScheduling;
diff --git a/Engine/Plugins/Media/BlackmagicMedia/Source/BlackmagicMediaOutput/Public/BlackmagicMediaOutput.h b/Engine/Plugins/Media/BlackmagicMedia/Source/BlackmagicMediaOutput/Public/BlackmagicMediaOutput.h
--- a/Engine/Plugins/Media/BlackmagicMedia/Source/BlackmagicMediaOutput/Public/BlackmagicMediaOutput.h
+++ b/Engine/Plugins/Media/BlackmagicMedia/Source/BlackmagicMediaOutput/Public/BlackmagicMediaOutput.h
@@ -103,6 +103,12 @@ public:
 	 */
 	UPROPERTY(EditAnywhere, AdvancedDisplay, Category = "Output")
 	bool bInterlacedFieldsTimecodeNeedToMatch;
+
+	/**
+	 * If this is enabled when outputting media in interlace format, frames will be treated as progressive and will be outputted as-is. (Experimental)
+	 */
+	UPROPERTY(EditAnywhere, AdvancedDisplay, Category = "Output")
+	bool bOutputInterlaceAsProgressive = false;
 	
 	/**
 	 * HDR Metadata of the video signal.
