Fix SDI Input video freeze after a few frames

If the video timecode format was set to auto, AJA's autodetect channel would acquire and release the stream for the whole application. Since we don't need exclusive access to the card if we just want to see the channel status, I've removed the Acquire and Release stream calls.

diff --git a/Engine/Plugins/Media/AjaMedia/Source/Aja/Private/AutoDetectChannel.cpp b/Engine/Plugins/Media/AjaMedia/Source/Aja/Private/AutoDetectChannel.cpp
--- a/Engine/Plugins/Media/AjaMedia/Source/Aja/Private/AutoDetectChannel.cpp
+++ b/Engine/Plugins/Media/AjaMedia/Source/Aja/Private/AutoDetectChannel.cpp
@@ -219,12 +219,6 @@ namespace AJA
 				::Sleep(SleepMilli);
 			}
 
-			if (!NewCard->AcquireStreamForApplicationWithReference('UE5.', static_cast<uint32_t>(AJAProcess::GetPid())))
-			{
-				UE_LOG(LogTemp,  Display, TEXT("AutoDetectChannel: Couldn't acquire stream for Unreal Engine application on device '%S'."), NewCard->GetDisplayName().c_str());
-				return false;
-			}
-
 			// save the service level.
 			NTV2EveryFrameTaskMode TaskMode;
 			NewCard->GetEveryFrameServices(TaskMode);
@@ -284,7 +278,6 @@ namespace AJA
 			}
 
 			NewCard->SetEveryFrameServices(TaskMode);
-			NewCard->ReleaseStreamForApplication('UE5.', static_cast<uint32_t>(AJAProcess::GetPid()));
 
 			return true;
 		}
