Don't try to build SSEMemCpy on arm64

#rnx

diff --git a/Engine/Plugins/Media/AjaMedia/Source/Aja/Private/SSE2MemCpy.cpp b/Engine/Plugins/Media/AjaMedia/Source/Aja/Private/SSE2MemCpy.cpp
--- a/Engine/Plugins/Media/AjaMedia/Source/Aja/Private/SSE2MemCpy.cpp
+++ b/Engine/Plugins/Media/AjaMedia/Source/Aja/Private/SSE2MemCpy.cpp
@@ -6,7 +6,7 @@
 #error SSE2_MEM_CPY already defined
 #endif
 
-#ifdef _WIN64
+#if defined(_WIN64) && !defined(_M_ARM64)
 #define SSE2_MEM_CPY
 #endif
 
@@ -34,7 +34,7 @@ namespace AJA
 #ifdef SSE2_MEM_CPY
 			return (((unsigned long long)(InDst) & 15) == 0) && (((unsigned long long)(InSrc) & 15) == 0) && ((InSize & 127) == 0);
 #else
-			return false
+			return false;
 #endif
 		}
 
diff --git a/Engine/Plugins/Media/BlackmagicMedia/Source/BlackmagicCore/Private/SSE2MemCpy.cpp b/Engine/Plugins/Media/BlackmagicMedia/Source/BlackmagicCore/Private/SSE2MemCpy.cpp
--- a/Engine/Plugins/Media/BlackmagicMedia/Source/BlackmagicCore/Private/SSE2MemCpy.cpp
+++ b/Engine/Plugins/Media/BlackmagicMedia/Source/BlackmagicCore/Private/SSE2MemCpy.cpp
@@ -6,7 +6,7 @@
 #error SSE2_MEM_CPY already defined
 #endif
 
-#if PLATFORM_WINDOWS
+#if PLATFORM_WINDOWS && !defined(_M_ARM64)
 	#define SSE2_MEM_CPY
 	#include <intrin.h>
 #elif PLATFORM_LINUX
