Aja - Reduce the amount of log spam when missing frames

diff --git a/Engine/Plugins/Media/AjaMedia/Source/Aja/Private/OutputChannel.cpp b/Engine/Plugins/Media/AjaMedia/Source/Aja/Private/OutputChannel.cpp
--- a/Engine/Plugins/Media/AjaMedia/Source/Aja/Private/OutputChannel.cpp
+++ b/Engine/Plugins/Media/AjaMedia/Source/Aja/Private/OutputChannel.cpp
@@ -937,6 +937,13 @@ namespace AJA
 					break;
 				}
 
+				if (bStopRequested)
+				{
+					// After the WaitForInputOrOutputInterrupt, it's possible that bStopRequested has been modified.
+					bRunning = false;
+					break;
+				}
+
 				if (!bIsValidWaitForOutput)
 				{
 					AJAAutoLock AutoLock(&Lock);
@@ -985,7 +992,7 @@ namespace AJA
 							Transfer.SetAudioBuffer(reinterpret_cast<ULWord*>(AvailableReadingFrame->AudioBuffer), AvailableReadingFrame->CopiedAudioBufferSize);
 						}
 
-						if (UseVideo() && AvailableReadingFrame->CopiedVideoBufferSize && TextureTransfer && GetOptions().bUseGPUDMA && !bStopRequested)
+						if (UseVideo() && AvailableReadingFrame->CopiedVideoBufferSize && TextureTransfer && GetOptions().bUseGPUDMA)
 						{
 							const UE::GPUTextureTransfer::ETransferDirection Direction = UE::GPUTextureTransfer::ETransferDirection::GPU_TO_CPU;
 							TextureTransfer->BeginSync(AvailableReadingFrame->VideoBuffer, Direction);
@@ -1006,7 +1013,7 @@ namespace AJA
 						}
 
 						AJA_CHECK(GetDevice().AutoCirculateTransfer(Channel, Transfer));
-						if (UseVideo() && AvailableReadingFrame->CopiedVideoBufferSize && TextureTransfer && GetOptions().bUseGPUDMA && !bStopRequested)
+						if (UseVideo() && AvailableReadingFrame->CopiedVideoBufferSize && TextureTransfer && GetOptions().bUseGPUDMA)
 						{
 							TextureTransfer->EndSync(AvailableReadingFrame->VideoBuffer);
 						}
@@ -1140,7 +1147,7 @@ namespace AJA
 			}
 		}
 
-		// As reference: this code was inspirered by NTV2LLBurn::ProcessFrames(void)
+		// As reference: this code was inspired by NTV2LLBurn::ProcessFrames(void)
 		void OutputChannelThread::Thread_PingPongLoop()
 		{
 			bool bHaveOutputOnce = false;
@@ -1154,6 +1161,13 @@ namespace AJA
 				TextureTransfer->ThreadPrep();
 			}
 
+			uint32_t LastFrameDropCount = 0;
+			auto LogDropFrames = [&LastFrameDropCount, this]() 
+			{
+				UE_LOG(LogAjaCore, Warning, TEXT("PingPong: Dropped %d frames while outputting on channel %d on device %S.\n"), LastFrameDropCount, uint32_t(Channel) + 1, GetDevice().GetDisplayName().c_str());
+			};
+
+
 			bool bRunning = true;
 			while (bRunning && !bStopRequested)
 			{
@@ -1179,6 +1193,13 @@ namespace AJA
 					break;
 				}
 
+				if (bStopRequested)
+				{
+					// After the WaitForInputOrOutputInterrupt, it's possible that bStopRequested has been modified.
+					bRunning = false;
+					break;
+				}
+
 				{
 					AJAAutoLock AutoLock(&Lock);
 					if (GetOptions().CallbackInterface)
@@ -1195,6 +1216,12 @@ namespace AJA
 				Frame* AvailableReadingFrame = Thread_FetchAvailableReadingFrame();
 				if (AvailableReadingFrame)
 				{
+					if (LastFrameDropCount != 0)
+					{
+						LogDropFrames();
+						LastFrameDropCount = 0;
+					}
+					
 					Thread_TestInterlacedOutput(AvailableReadingFrame);
 
 					//	Flip sense of the buffers again to refer to the buffers that the hardware isn't using (i.e. the off-screen buffers)...
@@ -1228,7 +1255,7 @@ namespace AJA
 					if (AvailableReadingFrame->CopiedVideoBufferSize > 0)
 					{
 						BurnTimecode(AvailableReadingFrame->Timecode, AvailableReadingFrame->VideoBuffer);
-						if (TextureTransfer && GetOptions().bUseGPUDMA && !bStopRequested)
+						if (TextureTransfer && GetOptions().bUseGPUDMA)
 						{
 							const UE::GPUTextureTransfer::ETransferDirection Direction = UE::GPUTextureTransfer::ETransferDirection::GPU_TO_CPU;
 							TextureTransfer->BeginSync(AvailableReadingFrame->VideoBuffer, Direction);
@@ -1236,7 +1263,7 @@ namespace AJA
 
 						bRunning = bRunning && GetDevice().DMAWriteFrame(BaseFrameIndex + CurrentOutFrame, reinterpret_cast<ULWord*>(AvailableReadingFrame->VideoBuffer), AvailableReadingFrame->CopiedVideoBufferSize, Channel);
 
-						if (TextureTransfer && GetOptions().bUseGPUDMA && !bStopRequested)
+						if (TextureTransfer && GetOptions().bUseGPUDMA)
 						{
 							TextureTransfer->EndSync(AvailableReadingFrame->VideoBuffer);
 						}
@@ -1289,13 +1316,21 @@ namespace AJA
 				}
 				else if (GetOptions().bDisplayWarningIfDropFrames)
 				{
+					TRACE_CPUPROFILER_EVENT_SCOPE(OutputChannel::NoFramesAvailable);
+
 					if (UseAudio() && !GetOptions().bDirectlyWriteAudio)
 					{
 						Thread_HandleLostFrameAudio();
 					}
 
-					TRACE_CPUPROFILER_EVENT_SCOPE(OutputChannel::NoFramesAvailable);
-					UE_LOG(LogAjaCore, Warning,  TEXT("PingPong: No frames are available for channel %d on device %S.\n"), uint32_t(Channel) + 1, GetDevice().GetDisplayName().c_str());
+					static constexpr uint32 NumMaxFrameBeforeWarning = 50;
+
+					LastFrameDropCount++;
+					if (LastFrameDropCount > NumMaxFrameBeforeWarning)
+					{
+						LogDropFrames();
+						LastFrameDropCount = 0;
+					}
 				}
 			}
 
diff --git a/Engine/Plugins/Media/MediaIOFramework/Source/MediaIOCore/Private/MediaCapture.cpp b/Engine/Plugins/Media/MediaIOFramework/Source/MediaIOCore/Private/MediaCapture.cpp
--- a/Engine/Plugins/Media/MediaIOFramework/Source/MediaIOCore/Private/MediaCapture.cpp
+++ b/Engine/Plugins/Media/MediaIOFramework/Source/MediaIOCore/Private/MediaCapture.cpp
@@ -62,7 +62,7 @@ static bool bSimulateError = false;
 FAutoConsoleCommand FMediaIOSetErrorState
 (
 	TEXT("MediaIO.SimulateError"),
-	TEXT("Sets all media captures to the error state."),
+	TEXT("Causes the next call to ProcessCapture_RenderThread to fail."),
 	FConsoleCommandDelegate::CreateLambda([]()
 		{
 			bSimulateError = true;
diff --git a/Engine/Plugins/Media/MediaIOFramework/Source/MediaIOCore/Public/MediaCapture.h b/Engine/Plugins/Media/MediaIOFramework/Source/MediaIOCore/Public/MediaCapture.h
--- a/Engine/Plugins/Media/MediaIOFramework/Source/MediaIOCore/Public/MediaCapture.h
+++ b/Engine/Plugins/Media/MediaIOFramework/Source/MediaIOCore/Public/MediaCapture.h
@@ -459,7 +459,7 @@ protected:
 	
 	virtual void StopCaptureImpl(bool bAllowPendingFrameToBeProcess) { }
 
-	/** Returns whether the capture implementation supports immediate (blocking) StopCapture.*/
+	/** Returns whether the capture implementation supports the implementation of StopCapture that returns a future.*/
 	virtual bool SupportsStopCapture_Future() const { return false; }
 
 	/**
@@ -469,7 +469,7 @@ protected:
 	 */
 	virtual TFuture<void> StopCapture_Future()
 	{
-		ensureMsgf(false, TEXT("Attempted to call StopCapture_Future without checking if the implementation supports it via SupportsImmediateStopCapture."));
+		ensureMsgf(false, TEXT("Attempted to call StopCapture_Future without checking if the implementation supports it via SupportsStopCapture_Future."));
 		return MakeFulfilledPromise<void>().GetFuture();
 	}
 
