Aja - Add an option to discard interlace frames if they land on an odd frame.

Add support for burning timecode on interlace frames, by burning the timecode of the two engine frames that were used to form the output frame.

diff --git a/Engine/Plugins/Media/AjaMedia/Source/Aja/Private/AJA.cpp b/Engine/Plugins/Media/AjaMedia/Source/Aja/Private/AJA.cpp
--- a/Engine/Plugins/Media/AjaMedia/Source/Aja/Private/AJA.cpp
+++ b/Engine/Plugins/Media/AjaMedia/Source/Aja/Private/AJA.cpp
@@ -11,6 +11,7 @@
 #include "AutoDetectChannel.h"
 #include "DeviceScanner.h"
 #include "InputChannel.h"
+#include "Misc/Timecode.h"
 #include "OutputChannel.h"
 #include "SyncChannel.h"
 #include "TimecodeChannel.h"
@@ -449,6 +450,11 @@ namespace AJA
 		return Channel ? Channel->DMAWriteAudio(Buffer, BufferSize) : false;
 	}
 
+	bool AJAOutputChannel::ShouldCaptureThisFrame(::FTimecode Timecode, uint32 SourceFrameNumber)
+	{
+		return Channel ? Channel->ShouldCaptureThisFrame(Timecode, SourceFrameNumber) : false;
+	}
+
 	bool AJAOutputChannel::GetOutputDimension(uint32_t& OutWidth, uint32_t& OutHeight) const
 	{
 		return Channel ? Channel->GetOutputDimension(OutWidth, OutHeight) : false;
diff --git a/Engine/Plugins/Media/AjaMedia/Source/Aja/Private/ChannelBase.cpp b/Engine/Plugins/Media/AjaMedia/Source/Aja/Private/ChannelBase.cpp
--- a/Engine/Plugins/Media/AjaMedia/Source/Aja/Private/ChannelBase.cpp
+++ b/Engine/Plugins/Media/AjaMedia/Source/Aja/Private/ChannelBase.cpp
@@ -743,5 +743,26 @@ namespace AJA
 				TimecodeBurner->BurnTimeCode(pVideoBuffer, TimeCodeString.c_str(), GetOptions().BurnTimecodePercentY);
 			}
 		}
+
+		void ChannelThreadBase::BurnTimecode(const AJA::FTimecode& InTimecode, const AJA::FTimecode& InTimecode2, uint8_t* InVideoBuffer)
+		{
+			if (TimecodeBurner)
+			{
+				const TimecodeFormat TcFormat = Helpers::ConvertToTimecodeFormat(VideoFormat);
+				CRP188 FrameRP188Info(InTimecode.Frames, InTimecode.Seconds, InTimecode.Minutes, InTimecode.Hours, TcFormat);
+
+				CRP188 FrameRP188Info2(InTimecode2.Frames, InTimecode2.Seconds, InTimecode2.Minutes, InTimecode2.Hours, TcFormat);
+
+				std::string TimeCodeString;
+				FrameRP188Info.GetRP188Str(TimeCodeString);
+				std::string TimeCodeString2;
+				FrameRP188Info2.GetRP188Str(TimeCodeString2);
+
+				char* pVideoBuffer = reinterpret_cast<char*>(InVideoBuffer + FormatDescriptor.firstActiveLine * FormatDescriptor.linePitch * 4);
+				TimecodeBurner->BurnTimeCode(pVideoBuffer, TimeCodeString, GetOptions().BurnTimecodePercentY);
+
+				TimecodeBurner->BurnTimeCode(pVideoBuffer, TimeCodeString2, 20);
+			}
+		}
 	}
 }
\ No newline at end of file
diff --git a/Engine/Plugins/Media/AjaMedia/Source/Aja/Private/ChannelBase.h b/Engine/Plugins/Media/AjaMedia/Source/Aja/Private/ChannelBase.h
--- a/Engine/Plugins/Media/AjaMedia/Source/Aja/Private/ChannelBase.h
+++ b/Engine/Plugins/Media/AjaMedia/Source/Aja/Private/ChannelBase.h
@@ -131,7 +131,10 @@ namespace AJA
 
 			void Thread_OnInitializationCompleted(bool bSucceed);
 
+			// Encode timecode onto a video buffer.
 			void BurnTimecode(const AJA::FTimecode& InTimecode, uint8_t* InVideoBuffer);
+			// Burn two timecodes on the frame (Used for interlace debugging to identify the timecode of the two frames that make up an output frame.)
+			void BurnTimecode(const AJA::FTimecode& InTimecode, const AJA::FTimecode& InTimecode2, uint8_t* InVideoBuffer);
 
 		private:
 			static void StaticThread(AJAThread* pThread, void* pContext);
diff --git a/Engine/Plugins/Media/AjaMedia/Source/Aja/Private/OutputChannel.cpp b/Engine/Plugins/Media/AjaMedia/Source/Aja/Private/OutputChannel.cpp
--- a/Engine/Plugins/Media/AjaMedia/Source/Aja/Private/OutputChannel.cpp
+++ b/Engine/Plugins/Media/AjaMedia/Source/Aja/Private/OutputChannel.cpp
@@ -121,6 +121,15 @@ namespace AJA
 			return false;
 		}
 
+		bool OutputChannel::ShouldCaptureThisFrame(::FTimecode Timeocode, uint32 SourceFrameNumber)
+		{
+			if (ChannelThread)
+			{
+				return ChannelThread->ShouldCaptureThisFrame(Timeocode, SourceFrameNumber);
+			}
+			return false;
+		}
+
 		bool OutputChannel::GetOutputDimension(uint32_t& OutWidth, uint32_t& OutHeight) const
 		{
 			NTV2FormatDescriptor FormatDescriptor = ChannelThread->FormatDescriptor;
@@ -325,12 +334,15 @@ namespace AJA
 
 				if (AvailableWritingFrame == nullptr && !bIsProgressive && !GetOptions().bOutputInterlacedFieldsTimecodeNeedToMatch)
 				{
-					if (!InFrameData.bEvenFrame)
+					if (!GetOptions().bOutputInterlaceOnEvenFrames)
 					{
-						// If the incoming frame is interlaced and we didn't find the correspondent even field
-						// do not create a new field. Drop it.
-						++LostFrameCounter;
-						return nullptr;
+						if (!InFrameData.bEvenFrame)
+						{
+							// If the incoming frame is interlaced and we didn't find the correspondent even field
+							// do not create a new field. Drop it.
+							++LostFrameCounter;
+							return nullptr;
+						}
 					}
 				}
 
@@ -427,6 +439,7 @@ namespace AJA
 					else
 					{
 						AvailableWritingFrame->FrameIdentifierF2 = InFrameData.FrameIdentifier;
+						AvailableWritingFrame->Timecode2 = InFrameData.Timecode;
 					}
 				}
 			}
@@ -759,6 +772,44 @@ namespace AJA
 			return true;
 		}
 
+		bool OutputChannelThread::ShouldCaptureThisFrame(::FTimecode Timecode, uint32 SourceFrameNumber)
+		{
+			const bool bInterlace = !::IsProgressivePicture(VideoFormat) && !Options.bOutputInterlaceAsProgressive;
+
+			if (bInterlace && Timecode.Frames % 2 != 0 && GetOptions().bOutputInterlaceOnEvenFrames)
+			{
+				// Odd frame should always have a matching frame.
+				// Find a frame where the even frame is filled but not the odd frame.
+				for (Frame* Frame : FrameReadyToWrite)
+				{
+					if (Frame->FrameIdentifier != AJAOutputFrameBufferData::InvalidFrameIdentifier && Frame->FrameIdentifierF2 == AJAOutputFrameBufferData::InvalidFrameIdentifier)
+					{
+						if (Frame->FrameIdentifier + 1 == SourceFrameNumber)
+						{
+							return true;
+						}
+					}
+				}
+
+				if (!FrameReadyToWrite.size())
+				{
+					UE_LOG(LogAjaCore, Verbose, TEXT("Skipping output because no writable frame is available."));
+				}
+				else
+				{
+					for (Frame* Frame : FrameReadyToWrite)
+					{
+						UE_LOG(LogAjaCore, VeryVerbose, TEXT("\t F1: %d, F2: %d"), Frame->FrameIdentifier, Frame->FrameIdentifierF2);
+					}
+				}
+
+				// Dont output an interlace frame starting with an odd frame.
+				return false;
+			}
+
+			return true;
+		}
+
 		bool OutputChannelThread::DeviceThread_ConfigureAnc(DeviceConnection::CommandList& InCommandList)
 		{
 			bool bResult = Super::DeviceThread_ConfigureAnc(InCommandList);
@@ -1000,7 +1051,14 @@ namespace AJA
 
 						if (UseVideo() && AvailableReadingFrame->CopiedVideoBufferSize)
 						{
-							BurnTimecode(AvailableReadingFrame->Timecode, AvailableReadingFrame->VideoBuffer);
+							if (bIsProgressive)
+							{
+								BurnTimecode(AvailableReadingFrame->Timecode, AvailableReadingFrame->VideoBuffer);
+							}
+							else
+							{
+								BurnTimecode(AvailableReadingFrame->Timecode, AvailableReadingFrame->Timecode2, AvailableReadingFrame->VideoBuffer);
+							}
 							Transfer.SetVideoBuffer(reinterpret_cast<ULWord*>(AvailableReadingFrame->VideoBuffer), AvailableReadingFrame->CopiedVideoBufferSize);
 						}
 
@@ -1254,7 +1312,14 @@ namespace AJA
 
 					if (AvailableReadingFrame->CopiedVideoBufferSize > 0)
 					{
-						BurnTimecode(AvailableReadingFrame->Timecode, AvailableReadingFrame->VideoBuffer);
+						if (bIsProgressive)
+						{
+							BurnTimecode(AvailableReadingFrame->Timecode, AvailableReadingFrame->VideoBuffer);
+						}
+						else
+						{
+							BurnTimecode(AvailableReadingFrame->Timecode, AvailableReadingFrame->Timecode2, AvailableReadingFrame->VideoBuffer);
+						}
 						if (TextureTransfer && GetOptions().bUseGPUDMA)
 						{
 							const UE::GPUTextureTransfer::ETransferDirection Direction = UE::GPUTextureTransfer::ETransferDirection::GPU_TO_CPU;
@@ -1477,7 +1542,7 @@ namespace AJA
 				bSuccess &= GetDevice().DMAWriteAudio(AudioSystem, reinterpret_cast<const ULWord*>(InBuffer), CurrentAudioWriteOffset, BytesUntilBufferEnd);
 				if (RemainingBytes > 0)
 				{
-					bSuccess &= GetDevice().DMAWriteAudio(AudioSystem, reinterpret_cast<const ULWord*>(InBuffer) + BytesUntilBufferEnd, 0, RemainingBytes);
+					bSuccess &= GetDevice().DMAWriteAudio(AudioSystem, reinterpret_cast<const ULWord*>(InBuffer + BytesUntilBufferEnd), 0, RemainingBytes);
 				}
 
 				CurrentAudioWriteOffset = RemainingBytes;
diff --git a/Engine/Plugins/Media/AjaMedia/Source/Aja/Private/OutputChannel.h b/Engine/Plugins/Media/AjaMedia/Source/Aja/Private/OutputChannel.h
--- a/Engine/Plugins/Media/AjaMedia/Source/Aja/Private/OutputChannel.h
+++ b/Engine/Plugins/Media/AjaMedia/Source/Aja/Private/OutputChannel.h
@@ -4,6 +4,7 @@
 
 #include "ChannelBase.h"
 #include "GPUTextureTransferModule.h"
+#include "Misc/Timecode.h"
 
 #include <vector>
 
@@ -33,6 +34,7 @@ namespace AJA
 			bool SetVideoFrameData(const AJAOutputFrameBufferData& InFrameData, uint8_t* InVideoBuffer, uint32_t InVideoBufferSize);
 			bool DMAWriteAudio(const uint8_t* InAudioBuffer, int32_t InAudioBufferSize);
 			bool SetVideoFrameData(const AJAOutputFrameBufferData& InFrameData, FRHITexture* RHITexture);
+			bool ShouldCaptureThisFrame(::FTimecode Timeocode, uint32 SourceFrameNumber);
 			bool GetOutputDimension(uint32_t& OutWidth, uint32_t& OutHeight) const;
 			int32_t GetNumAudioSamplesPerFrame(const AJAOutputFrameBufferData& InFrameData) const;
 
@@ -59,6 +61,7 @@ namespace AJA
 			bool SetAudioFrameData(const AJAOutputFrameBufferData& InFrameData, uint8_t* AudioBuffer, uint32_t AudioBufferSize);
 			bool SetVideoFrameData(const AJAOutputFrameBufferData& InFrameData, uint8_t* VideoBuffer, uint32_t VideoBufferSize);
 			bool SetVideoFrameData(const AJAOutputFrameBufferData& InFrameData, FRHITexture* RHITexture);
+			bool ShouldCaptureThisFrame(::FTimecode Timecode, uint32 SourceFrameNumber);
 			bool DMAWriteAudio(const uint8_t* InAudioBuffer, int32_t BufferSize);
 
 		protected:
@@ -110,6 +113,7 @@ namespace AJA
 				uint32_t FrameIdentifier;
 				uint32_t FrameIdentifierF2; // used in interlaced
 				FTimecode Timecode;
+				FTimecode Timecode2;
 
 				bool bAncLineFilled;
 				bool bAncF2LineFilled; // used in interlaced
diff --git a/Engine/Plugins/Media/AjaMedia/Source/Aja/Public/AJALib.h b/Engine/Plugins/Media/AjaMedia/Source/Aja/Public/AJALib.h
--- a/Engine/Plugins/Media/AjaMedia/Source/Aja/Public/AJALib.h
+++ b/Engine/Plugins/Media/AjaMedia/Source/Aja/Public/AJALib.h
@@ -8,6 +8,7 @@
 
 #include "AjaCoreModule.h"
 #include "Async/Future.h"
+#include "Misc/Timecode.h"
 
 
 class FRHITexture;
@@ -461,6 +462,7 @@ namespace AJA
 				uint32_t bUseKey : 1; // output will also sent the key on OutputKeyPortIndex
 				uint32_t bOutputInterlacedFieldsTimecodeNeedToMatch : 1; // when trying to find the odd field that correspond to the even field, the 2 timecode need to match
 				uint32_t bOutputInterlaceAsProgressive : 1; // whether to treat outgoing interlace frames as progressive.
+				uint32_t bOutputInterlaceOnEvenFrames : 1; // whether to force output interlace on even timecode frames.
 				uint32_t bUseAncillary : 1; // enable ANC system
 				uint32_t bUseAudio : 1; // enable audio input/output
 				uint32_t bUseVideo : 1; // enable video input/output
@@ -540,6 +542,9 @@ namespace AJA
 		bool SetVideoFrameData(const AJAOutputFrameBufferData& InFrameData, FRHITexture* RHITexture);
 		bool DMAWriteAudio(const uint8* Buffer, int32 BufferSize);
 
+		// Return whether to capture or discard a frame based on its timecode and frame number.
+		bool ShouldCaptureThisFrame(::FTimecode Timecode, uint32 SourceFrameNumber);
+
 		bool GetOutputDimension(uint32_t& OutWidth, uint32_t& OutHeight) const;
 		int32_t GetNumAudioSamplesPerFrame(const AJAOutputFrameBufferData& InFrameData) const;
 
diff --git a/Engine/Plugins/Media/AjaMedia/Source/AjaMediaOutput/Private/AjaMediaCapture.cpp b/Engine/Plugins/Media/AjaMedia/Source/AjaMediaOutput/Private/AjaMediaCapture.cpp
--- a/Engine/Plugins/Media/AjaMedia/Source/AjaMediaOutput/Private/AjaMediaCapture.cpp
+++ b/Engine/Plugins/Media/AjaMedia/Source/AjaMediaOutput/Private/AjaMediaCapture.cpp
@@ -379,6 +379,15 @@ bool UAjaMediaCapture::ShouldCaptureRHIResource() const
 	return TextureTransfer && FGPUTextureTransferModule::Get().IsEnabled();
 }
 
+bool UAjaMediaCapture::ShouldCaptureThisFrame(const FCaptureBaseData& InBaseData) const
+{
+	if (OutputChannel)
+	{
+		return OutputChannel->ShouldCaptureThisFrame(InBaseData.SourceFrameTimecode, InBaseData.SourceFrameNumber);
+	}
+	return false;
+}
+
 void UAjaMediaCapture::ApplyViewportTextureAlpha(TSharedPtr<FSceneViewport> InSceneViewport)
 {
 	if (InSceneViewport.IsValid())
@@ -517,6 +526,7 @@ bool UAjaMediaCapture::InitAJA(UAjaMediaOutput* InAjaMediaOutput)
 	ChannelOptions.bUseVideo = true;
 	ChannelOptions.bOutputInterlacedFieldsTimecodeNeedToMatch = InAjaMediaOutput->bInterlacedFieldsTimecodeNeedToMatch && Descriptor.bIsInterlacedStandard && InAjaMediaOutput->TimecodeFormat != EMediaIOTimecodeFormat::None;
 	ChannelOptions.bOutputInterlaceAsProgressive = InAjaMediaOutput->bOutputInterlaceAsProgressive;
+	ChannelOptions.bOutputInterlaceOnEvenFrames = InAjaMediaOutput->bOutputInterlaceOnEvenFrames;
 	ChannelOptions.bDisplayWarningIfDropFrames = bLogDropFrame;
 	ChannelOptions.bConvertOutputLevelAToB = InAjaMediaOutput->bOutputIn3GLevelB && Descriptor.bIsVideoFormatA;
 	ChannelOptions.TransportType = AjaMediaCaptureUtils::ConvertTransportType(InAjaMediaOutput->OutputConfiguration.MediaConfiguration.MediaConnection.TransportType, InAjaMediaOutput->OutputConfiguration.MediaConfiguration.MediaConnection.QuadTransportType);
diff --git a/Engine/Plugins/Media/AjaMedia/Source/AjaMediaOutput/Public/AjaMediaCapture.h b/Engine/Plugins/Media/AjaMedia/Source/AjaMediaOutput/Public/AjaMediaCapture.h
--- a/Engine/Plugins/Media/AjaMedia/Source/AjaMediaOutput/Public/AjaMediaCapture.h
+++ b/Engine/Plugins/Media/AjaMedia/Source/AjaMediaOutput/Public/AjaMediaCapture.h
@@ -52,6 +52,7 @@ protected:
 	virtual bool SupportsStopCapture_Future() const override { return true; }
 	virtual TFuture<void> StopCapture_Future() override;
 	virtual bool ShouldCaptureRHIResource() const override;
+	virtual bool ShouldCaptureThisFrame(const FCaptureBaseData& InBaseData) const override;
 	
 	virtual void OnFrameCaptured_RenderingThread(const FCaptureBaseData& InBaseData, TSharedPtr<FMediaCaptureUserData, ESPMode::ThreadSafe> InUserData, void* InBuffer, int32 Width, int32 Height, int32 BytesPerRow) override;
 	virtual void OnRHIResourceCaptured_RenderingThread(FRHICommandListImmediate& RHICmdList, const FCaptureBaseData& InBaseData, TSharedPtr<FMediaCaptureUserData, ESPMode::ThreadSafe> InUserData, FTextureRHIRef InTexture) override;
diff --git a/Engine/Plugins/Media/AjaMedia/Source/AjaMediaOutput/Public/AjaMediaOutput.h b/Engine/Plugins/Media/AjaMedia/Source/AjaMediaOutput/Public/AjaMediaOutput.h
--- a/Engine/Plugins/Media/AjaMedia/Source/AjaMediaOutput/Public/AjaMediaOutput.h
+++ b/Engine/Plugins/Media/AjaMedia/Source/AjaMediaOutput/Public/AjaMediaOutput.h
@@ -118,6 +118,13 @@ public:
 	UPROPERTY(EditAnywhere, AdvancedDisplay, Category = "Output")
 	bool bOutputInterlaceAsProgressive = false;
 
+	/**
+	 * Whether to only schedule interlace frames if we land on an even timecode. Useful when you have an interlace input
+	 * and want to ensure that we interlace the original frames. (Experimental)
+	 */
+	UPROPERTY(EditAnywhere, AdvancedDisplay, Category = "Output")
+	bool bOutputInterlaceOnEvenFrames = false;
+
 	/**
 	 * HDR Metadata of the video signal.
 	 */
