Aja - Add option to continue input, output and genlock when card timeouts

diff --git a/Engine/Plugins/Media/AjaMedia/Source/Aja/Private/InputChannel.cpp b/Engine/Plugins/Media/AjaMedia/Source/Aja/Private/InputChannel.cpp
--- a/Engine/Plugins/Media/AjaMedia/Source/Aja/Private/InputChannel.cpp
+++ b/Engine/Plugins/Media/AjaMedia/Source/Aja/Private/InputChannel.cpp
@@ -448,7 +448,7 @@ namespace AJA
 					bool bWaitResult = Device->WaitForInputOrOutputInterrupt(Channel);
 					if (!bWaitResult)
 					{
-						if (!Options.bAutoDetectFormat)
+						if (!Options.bAutoDetectFormat && Options.bStopOnTimeout)
 						{
 							AJA_CHECK(false);
 						}
@@ -463,7 +463,7 @@ namespace AJA
 					EVerifyFormatResult Result = VerifyFormat();
 					if (Result == EVerifyFormatResult::Failure)
 					{
-						if (!GetOptions().bAutoDetectFormat)
+						if (!GetOptions().bAutoDetectFormat && GetOptions().bStopOnTimeout)
 						{
 							bRunning = false;
 							break;
@@ -553,7 +553,7 @@ namespace AJA
 				bool bWaitResult = Device->WaitForInputOrOutputInterrupt(Channel);
 				if (!bWaitResult)
 				{
-					if (!Options.bAutoDetectFormat)
+					if (!Options.bAutoDetectFormat && Options.bStopOnTimeout)
 					{
 						bRunning = false;
 						AJA_CHECK(false);
diff --git a/Engine/Plugins/Media/AjaMedia/Source/Aja/Private/OutputChannel.cpp b/Engine/Plugins/Media/AjaMedia/Source/Aja/Private/OutputChannel.cpp
--- a/Engine/Plugins/Media/AjaMedia/Source/Aja/Private/OutputChannel.cpp
+++ b/Engine/Plugins/Media/AjaMedia/Source/Aja/Private/OutputChannel.cpp
@@ -978,6 +978,14 @@ namespace AJA
 				if (!bRunning)
 				{
 					UE_LOG(LogAjaCore, Error, TEXT("AutoCirculate: Can't wait for the output field for channel %d on device %S.\n"), uint32_t(Channel) + 1, GetDevice().GetDisplayName().c_str());
+					if (GetOptions().bStopOnTimeout)
+					{
+						break;
+					}
+					else
+					{
+						bRunning = true;
+					}
 				}
 
 				AUTOCIRCULATE_STATUS ChannelStatus;
@@ -1248,7 +1256,14 @@ namespace AJA
 				if (!bRunning)
 				{
 					UE_LOG(LogAjaCore, Error, TEXT("PingPong: Can't wait for the output field for channel %d on device %S.\n"), uint32_t(Channel) + 1, GetDevice().GetDisplayName().c_str());
-					break;
+					if (GetOptions().bStopOnTimeout)
+					{
+						break;
+					}
+					else
+					{
+						bRunning = true;
+					}
 				}
 
 				if (bStopRequested)
diff --git a/Engine/Plugins/Media/AjaMedia/Source/Aja/Private/SyncChannel.cpp b/Engine/Plugins/Media/AjaMedia/Source/Aja/Private/SyncChannel.cpp
--- a/Engine/Plugins/Media/AjaMedia/Source/Aja/Private/SyncChannel.cpp
+++ b/Engine/Plugins/Media/AjaMedia/Source/Aja/Private/SyncChannel.cpp
@@ -341,6 +341,13 @@ namespace AJA
 				bResult = Device->WaitForInputOrOutputInterrupt(Channel, 1, InterruptType);
 			}
 
+			if (!bResult && !SyncOption.bStopOnCardTimeout)
+			{
+				// If bStopOnCardTimeout is disabled, then we want sync to keep going if the SDK timeouts while waiting for a card interupt.
+				UE_LOG(LogAjaCore, Error, TEXT("SyncChannel: Was not able to wait for an interrupt signal on device %S for channel %d.\n"), GetDevice().GetDisplayName().c_str(), uint32_t(Channel) + 1);
+				bResult = true;
+			}
+
 			return bResult;
 		}
 
diff --git a/Engine/Plugins/Media/AjaMedia/Source/Aja/Public/AJALib.h b/Engine/Plugins/Media/AjaMedia/Source/Aja/Public/AJALib.h
--- a/Engine/Plugins/Media/AjaMedia/Source/Aja/Public/AJALib.h
+++ b/Engine/Plugins/Media/AjaMedia/Source/Aja/Public/AJALib.h
@@ -219,6 +219,7 @@ namespace AJA
 		bool bOutput; // port is output
 		bool bWaitForFrameToBeReady; // port is input and we want to wait for the image to be sent to Unreal Engine before ticking
 		bool bAutoDetectFormat;
+		bool bStopOnCardTimeout = true;  // Stop sync channel whenever we timeout while waiting for a vertical interrupt from the card (Default: true)
 	};
 
 	class AJACORE_API AJASyncChannel
@@ -473,6 +474,7 @@ namespace AJA
 				uint32_t bUseGPUDMA : 1; // Whether to use GPU Direct when outputting.
 				uint32_t bAutoDetectFormat : 1; // Whether to autodetect format on input.
 				uint32_t bDirectlyWriteAudio : 1; // Experimental mode to set audio from the audio thread directly.
+				uint32_t bStopOnTimeout : 1; // Stop input/output whenever we timeout while waiting for a vertical interrupt from the card.
 			};
 			uint32_t Options;
 		};
diff --git a/Engine/Plugins/Media/AjaMedia/Source/AjaMedia/Private/AjaMediaPrivate.h b/Engine/Plugins/Media/AjaMedia/Source/AjaMedia/Private/AjaMediaPrivate.h
--- a/Engine/Plugins/Media/AjaMedia/Source/AjaMedia/Private/AjaMediaPrivate.h
+++ b/Engine/Plugins/Media/AjaMedia/Source/AjaMedia/Private/AjaMediaPrivate.h
@@ -26,6 +26,7 @@ namespace AjaMediaOption
 	static const FName LogDropFrame("LogDropFrame");
 	static const FName EncodeTimecodeInTexel("EncodeTimecodeInTexel");
 	static const FName CaptureWithAutoCirculating("CaptureWithAutoCirculating");
+	static const FName StopOnTimeout("StopOnTimeout");
 	static const FName CaptureAncillary("CaptureAncillary");
 	static const FName CaptureAudio("CaptureAudio");
 	static const FName CaptureVideo("CaptureVideo");
diff --git a/Engine/Plugins/Media/AjaMedia/Source/AjaMedia/Private/Assets/AjaCustomTimeStep.cpp b/Engine/Plugins/Media/AjaMedia/Source/AjaMedia/Private/Assets/AjaCustomTimeStep.cpp
--- a/Engine/Plugins/Media/AjaMedia/Source/AjaMedia/Private/Assets/AjaCustomTimeStep.cpp
+++ b/Engine/Plugins/Media/AjaMedia/Source/AjaMedia/Private/Assets/AjaCustomTimeStep.cpp
@@ -196,6 +196,7 @@ bool UAjaCustomTimeStep::Initialize_Internal(UEngine* InEngine)
 	Options.bWaitForFrameToBeReady = bWaitForFrameToBeReady && !bUseReferenceIn;
 	Options.TransportType = AJA::ETransportType::TT_SdiSingle;
 	Options.bAutoDetectFormat = bAutoDetectFormat;
+	Options.bStopOnCardTimeout = bStopOnCardTimeout;
 
 	{
 		const EMediaIOTransportType TransportType = MediaConfiguration.MediaConnection.TransportType;
diff --git a/Engine/Plugins/Media/AjaMedia/Source/AjaMedia/Private/Assets/AjaMediaSource.cpp b/Engine/Plugins/Media/AjaMedia/Source/AjaMedia/Private/Assets/AjaMediaSource.cpp
--- a/Engine/Plugins/Media/AjaMedia/Source/AjaMedia/Private/Assets/AjaMediaSource.cpp
+++ b/Engine/Plugins/Media/AjaMedia/Source/AjaMedia/Private/Assets/AjaMediaSource.cpp
@@ -55,6 +55,10 @@ bool UAjaMediaSource::GetMediaOption(const FName& Key, bool DefaultValue) const
 	{
 		return bEncodeTimecodeInTexel;
 	}
+	if (Key == AjaMediaOption::StopOnTimeout)
+	{
+		return bStopInputOnCardTimeout;
+	}
 
 	return Super::GetMediaOption(Key, DefaultValue);
 }
diff --git a/Engine/Plugins/Media/AjaMedia/Source/AjaMedia/Private/Player/AjaMediaPlayer.cpp b/Engine/Plugins/Media/AjaMedia/Source/AjaMedia/Private/Player/AjaMediaPlayer.cpp
--- a/Engine/Plugins/Media/AjaMedia/Source/AjaMedia/Private/Player/AjaMediaPlayer.cpp
+++ b/Engine/Plugins/Media/AjaMedia/Source/AjaMedia/Private/Player/AjaMediaPlayer.cpp
@@ -861,6 +861,7 @@ bool FAjaMediaPlayer::Open_Internal(const FString& Url, const IMediaOptions* Opt
 		AjaOptions.bBurnTimecode = false;
 		AjaOptions.BurnTimecodePercentY = 80;
 		AjaOptions.bAutoDetectFormat = bAutoDetect;
+		AjaOptions.bStopOnTimeout = Options->GetMediaOption(AjaMediaOption::StopOnTimeout, true);
 
 		//Adjust supported sample types based on what's being captured
 		SupportedSampleTypes = AjaOptions.bUseVideo ? EMediaIOSampleType::Video : EMediaIOSampleType::None;
diff --git a/Engine/Plugins/Media/AjaMedia/Source/AjaMedia/Public/AjaCustomTimeStep.h b/Engine/Plugins/Media/AjaMedia/Source/AjaMedia/Public/AjaCustomTimeStep.h
--- a/Engine/Plugins/Media/AjaMedia/Source/AjaMedia/Public/AjaCustomTimeStep.h
+++ b/Engine/Plugins/Media/AjaMedia/Source/AjaMedia/Public/AjaCustomTimeStep.h
@@ -68,6 +68,13 @@ public:
 	UPROPERTY(EditAnywhere, Category="Genlock", meta=(DisplayName="Configuration"))
 	FMediaIOConfiguration MediaConfiguration;
 
+	/**
+	 * If this is disabled, genlock will attempt to recover if the SDK timeouts while waiting for an interrupt signal from the card. It will remain in the synchronizing state until the card receives a signal.
+	 * Note: In editor outside of PIE and -game, the CustomTimeStep will try to restart itself automatically regardless of this setting.
+	 */
+	UPROPERTY(EditAnywhere, Category = "Genlock")
+	bool bStopOnCardTimeout = true;
+
 	/**
 	 * If true, the Engine will wait for the frame to be read.
 	 * This will introduce random latency (the time it takes to read a frame).
diff --git a/Engine/Plugins/Media/AjaMedia/Source/AjaMedia/Public/AjaMediaSource.h b/Engine/Plugins/Media/AjaMedia/Source/AjaMedia/Public/AjaMediaSource.h
--- a/Engine/Plugins/Media/AjaMedia/Source/AjaMedia/Public/AjaMediaSource.h
+++ b/Engine/Plugins/Media/AjaMedia/Source/AjaMedia/Public/AjaMediaSource.h
@@ -75,6 +75,12 @@ public:
 	UPROPERTY(BlueprintReadOnly, EditAnywhere, AdvancedDisplay, Category="Ancillary", meta=(EditCondition="bCaptureAncillary", ClampMin="1", ClampMax="32"))
 	int32 MaxNumAncillaryFrameBuffer;
 
+	/**
+	 * If this is disabled, input will keep going even if we hit a timeout when waiting for a vertical interrupt from the card. 
+	 */
+	UPROPERTY(EditAnywhere, Category = "AJA")
+	bool bStopInputOnCardTimeout = true;
+
 public:
 	/** Capture Audio from the AJA source. */
 	UPROPERTY(BlueprintReadOnly, EditAnywhere, Category="Audio")
diff --git a/Engine/Plugins/Media/AjaMedia/Source/AjaMediaOutput/Private/AjaMediaCapture.cpp b/Engine/Plugins/Media/AjaMedia/Source/AjaMediaOutput/Private/AjaMediaCapture.cpp
--- a/Engine/Plugins/Media/AjaMedia/Source/AjaMediaOutput/Private/AjaMediaCapture.cpp
+++ b/Engine/Plugins/Media/AjaMedia/Source/AjaMediaOutput/Private/AjaMediaCapture.cpp
@@ -535,6 +535,7 @@ bool UAjaMediaCapture::InitAJA(UAjaMediaOutput* InAjaMediaOutput)
 	ChannelOptions.OutputReferenceType = AjaMediaCaptureUtils::Convert(InAjaMediaOutput->OutputConfiguration.OutputReference);
 	ChannelOptions.bUseGPUDMA = ShouldCaptureRHIResource();
 	ChannelOptions.bDirectlyWriteAudio = InAjaMediaOutput->bOutputAudioOnAudioThread;
+	ChannelOptions.bStopOnTimeout = InAjaMediaOutput->bStopOutputOnCardTimeout;
 
 	ChannelOptions.HDROptions = AjaMediaCaptureUtils::MakeAjaHDRMetadata(InAjaMediaOutput->HDROptions);
 	
diff --git a/Engine/Plugins/Media/AjaMedia/Source/AjaMediaOutput/Public/AjaMediaOutput.h b/Engine/Plugins/Media/AjaMedia/Source/AjaMediaOutput/Public/AjaMediaOutput.h
--- a/Engine/Plugins/Media/AjaMedia/Source/AjaMediaOutput/Public/AjaMediaOutput.h
+++ b/Engine/Plugins/Media/AjaMedia/Source/AjaMediaOutput/Public/AjaMediaOutput.h
@@ -52,6 +52,12 @@ public:
 	UPROPERTY(EditAnywhere, Category="AJA", meta=(DisplayName="Configuration"))
 	FMediaIOOutputConfiguration OutputConfiguration;
 
+	/**
+	 * If this is disabled, output will keep going even if we hit a timeout when waiting for a vertical interrupt from the card.
+	 */
+	UPROPERTY(EditAnywhere, Category = "AJA")
+	bool bStopOutputOnCardTimeout = true;
+
 public:
 	/**
 	 * The output of the Audio, Ancillary and/or video will be perform at the same time.
