Aja - Add a new output mode that can reduce latency by up to 1 frame.

The new output mode can be enabled with Aja.PingPongVersion 1 (Only supported for PingPong Progressive output).

New UnrealInsights debug info can be enabled with Aja.EnableDebugMarkers 1

In this mode, the output thread avoids waiting for the card?s vertical blank interrupt (VBI) when possible. Instead, it tracks the current ?out? frame and writes directly to the buffer if the card?s buffer is available, rather than always stalling until the next VBI.

Another improvement in the output thread is that it now waits for a frame to be pushed by Unreal before resuming. This eliminates situations where the thread would wake up, wait for a VBI, check the queue, and then stalls for an additional VBI period before checking again.

diff --git a/Engine/Plugins/Media/AjaMedia/Source/Aja/Private/Device.cpp b/Engine/Plugins/Media/AjaMedia/Source/Aja/Private/Device.cpp
--- a/Engine/Plugins/Media/AjaMedia/Source/Aja/Private/Device.cpp
+++ b/Engine/Plugins/Media/AjaMedia/Source/Aja/Private/Device.cpp
@@ -53,7 +53,7 @@ namespace AJA
 		{
 		}
 
-		bool DeviceConnection::SyncThreadInfoThread::WaitForVerticalInterrupt(DeviceConnection* InDeviceConnection, ChannelInfo* InChannelInfo)
+		bool DeviceConnection::SyncThreadInfoThread::WaitForVerticalInterrupt(DeviceConnection* InDeviceConnection, ChannelInfo* InChannelInfo, uint32& OutSyncCount)
 		{
 			assert(InDeviceConnection);
 			assert(InChannelInfo);
@@ -82,6 +82,7 @@ namespace AJA
 						bSecondCountValid = InDeviceConnection->GetCard()->GetOutputVerticalInterruptCount(AfterCount, InChannelInfo->Channel);
 					} while (bFirstCountValid && bSecondCountValid && BeforeCount == AfterCount && bResult);
 				}
+				OutSyncCount = AfterCount;
 			}
 
 			return bResult;
@@ -120,9 +121,12 @@ namespace AJA
 		bool DeviceConnection::SyncThreadInfoThread::ThreadLoop()
 		{
 			// do the wait
-			bWaitResult = WaitForVerticalInterrupt(DeviceConnectionPtr, ChannelInfoPtr);
+			uint32 NewSyncCount = 0;
+			bWaitResult = WaitForVerticalInterrupt(DeviceConnectionPtr, ChannelInfoPtr, NewSyncCount);
 			SyncCondition.notify_all();
 
+			OnVerticalInterruptDelegate.Broadcast(NewSyncCount);
+
 			return true;
 		}
 
@@ -708,6 +712,47 @@ namespace AJA
 			CommandThread.Kill(-1);
 		}
 
+		FDelegateHandle DeviceConnection::SetOnInterruptEvent(NTV2Channel InChannel, AJA::Private::DeviceConnection::FOnVerticalInterrupt::FDelegate OnInterrupt)
+		{
+			FDelegateHandle DelegateHandle;
+
+			AJAAutoLock Lock(&ChannelInfoLock);
+			auto FoundIt = std::find_if(std::begin(ChannelInfos), std::end(ChannelInfos), [=](const ChannelInfo* Info) { return Info->Channel == InChannel; });
+
+			if (FoundIt != std::end(ChannelInfos))
+			{
+				ChannelInfo* FoundChannel = *FoundIt;
+
+				if (FoundChannel->SyncThread)
+				{
+					DelegateHandle = FoundChannel->SyncThread->OnVerticalInterrupt().Add(MoveTemp(OnInterrupt));
+				}
+				else
+				{
+					UE_LOG(LogAjaCore, Error, TEXT("No sync thread available when calling SetOnInterruptEvent."));
+				}
+			}
+
+			return DelegateHandle;
+		}
+
+		void DeviceConnection::ClearOnInterruptEvent(NTV2Channel InChannel, FDelegateHandle DelegateHandle)
+		{
+			AJAAutoLock Lock(&ChannelInfoLock);
+			auto FoundIt = std::find_if(std::begin(ChannelInfos), std::end(ChannelInfos), [=](const ChannelInfo* Info) { return Info->Channel == InChannel; });
+			if (FoundIt == std::end(ChannelInfos))
+			{
+				return;
+			}
+
+			ChannelInfo* FoundChannel = *FoundIt;
+
+			if (FoundChannel->SyncThread)
+			{
+				FoundChannel->SyncThread->OnVerticalInterrupt().Remove(DelegateHandle);
+			}
+		}
+
 		bool DeviceConnection::WaitForInputOrOutputInterrupt(NTV2Channel InChannel, int32_t InRepeatCount, EInterruptType InInterrupt)
 		{
 			// NTV2Card::WaitForInputVerticalInterrupt can wait for only one thread at the same time
diff --git a/Engine/Plugins/Media/AjaMedia/Source/Aja/Private/Device.h b/Engine/Plugins/Media/AjaMedia/Source/Aja/Private/Device.h
--- a/Engine/Plugins/Media/AjaMedia/Source/Aja/Private/Device.h
+++ b/Engine/Plugins/Media/AjaMedia/Source/Aja/Private/Device.h
@@ -33,6 +33,7 @@ namespace AJA
 			static const uint32_t NumberOfFrameToAquire = 4;
 			static const uint32_t NumberOfFrameForAutoCirculate = NumberOfFrameToAquire;
 			static const uint32_t InvalidFrameBufferIndex = 0xFFFFFFFF;
+			DECLARE_TS_MULTICAST_DELEGATE_OneParam(FOnVerticalInterrupt, uint32 /*NewSyncCount*/);
 
 		private:
 			struct ChannelInfo;
@@ -42,6 +43,8 @@ namespace AJA
 			{
 			public:
 				SyncThreadInfoThread(DeviceConnection* InDeviceConnection, ChannelInfo* InChannelInfo);
+				/** Delegate triggered whenever the VBI count has increased. */
+				FOnVerticalInterrupt& OnVerticalInterrupt() { return OnVerticalInterruptDelegate; }
 
 			private:
 				ChannelInfo* ChannelInfoPtr;
@@ -50,9 +53,10 @@ namespace AJA
 				std::mutex SyncMutex;
 				std::condition_variable SyncCondition;
 				volatile bool bWaitResult;
+				FOnVerticalInterrupt OnVerticalInterruptDelegate;
 
 			public:
-				bool WaitForVerticalInterrupt(DeviceConnection* InDeviceConnection, ChannelInfo* InChannelInfo);
+				bool WaitForVerticalInterrupt(DeviceConnection* InDeviceConnection, ChannelInfo* InChannelInfo, uint32& OutSyncCount);
 				static bool IsCurrentField(DeviceConnection* InDeviceConnection, ChannelInfo* InChannelInfo, NTV2FieldID InFieldId = NTV2_FIELD0);
 				bool Wait_ExternalThread();
 				virtual bool ThreadLoop() override;
@@ -104,6 +108,13 @@ namespace AJA
 				InputField,
 				Auto,
 			};
+
+			/**
+			 * Set a callback to be invoked whenever the card sends a vertical interrupt signal.
+			 * @return a valid Delegate if the delegate was set correctly.
+			 */
+			FDelegateHandle SetOnInterruptEvent(NTV2Channel InChannel, AJA::Private::DeviceConnection::FOnVerticalInterrupt::FDelegate OnInterrupt);
+			void ClearOnInterruptEvent(NTV2Channel InChannel, FDelegateHandle DelegateHandle);
 			bool WaitForInputOrOutputInterrupt(NTV2Channel InChannel, int32_t InRepeatCount = 1, EInterruptType InInterrupt = EInterruptType::Auto);
 			bool WaitForInputOrOutputInputField(NTV2Channel InChannel, int32_t InRepeatCount, NTV2FieldID& OutInputField);
 			NTV2VideoFormat GetVideoFormat(NTV2Channel InChannel);
diff --git a/Engine/Plugins/Media/AjaMedia/Source/Aja/Private/Helpers.cpp b/Engine/Plugins/Media/AjaMedia/Source/Aja/Private/Helpers.cpp
--- a/Engine/Plugins/Media/AjaMedia/Source/Aja/Private/Helpers.cpp
+++ b/Engine/Plugins/Media/AjaMedia/Source/Aja/Private/Helpers.cpp
@@ -251,6 +251,26 @@ namespace AJA
 			return NTV2_FRAMERATE_UNKNOWN;
 		}
 
+		FFrameRate Helpers::ConvertToUnrealFramerate(NTV2VideoFormat InVideoFormat)
+		{
+			FFrameRate Result;
+			switch (::GetNTV2FrameRateFromVideoFormat(InVideoFormat))
+			{
+			case NTV2_FRAMERATE_6000:	Result = FFrameRate(60,1);	break;
+			case NTV2_FRAMERATE_5994:	Result = FFrameRate(60000,1001);	break;
+			case NTV2_FRAMERATE_4800:	Result = FFrameRate(48,1);	break;
+			case NTV2_FRAMERATE_4795:	Result = FFrameRate(48000, 1001);	break;
+			case NTV2_FRAMERATE_3000:	Result = FFrameRate(30,1);	break;
+			case NTV2_FRAMERATE_2997:	Result = FFrameRate(30000, 1001);	break;
+			case NTV2_FRAMERATE_2500:	Result = FFrameRate(25, 1);	break;
+			case NTV2_FRAMERATE_2400:	Result = FFrameRate(24, 1);	break;
+			case NTV2_FRAMERATE_2398:	Result = FFrameRate(24000, 1001);	break;
+			case NTV2_FRAMERATE_5000:	Result = FFrameRate(50, 1);	break;
+			default:					break;
+			}
+			return Result;
+		}
+
 		TimecodeFormat Helpers::ConvertToTimecodeFormat(NTV2VideoFormat InVideoFormat)
 		{
 			TimecodeFormat Result = kTCFormatUnknown;
diff --git a/Engine/Plugins/Media/AjaMedia/Source/Aja/Private/Helpers.h b/Engine/Plugins/Media/AjaMedia/Source/Aja/Private/Helpers.h
--- a/Engine/Plugins/Media/AjaMedia/Source/Aja/Private/Helpers.h
+++ b/Engine/Plugins/Media/AjaMedia/Source/Aja/Private/Helpers.h
@@ -4,6 +4,7 @@
 
 #include "CoreMinimal.h"
 #include "AjaCoreModule.h"
+#include "Misc/FrameRate.h"
 
 #define MSWindows 1
 #define AJA_WINDOWS 1
@@ -106,6 +107,7 @@ namespace AJA
 			static EAjaHDRMetadataLuminance ConvertFromAjaHDRLuminance(NTV2HDRLuminance AjaHDRLuminance);
 			
 			static NTV2FrameRate ConvertToFrameRate(uint32_t InNumerator, uint32_t InDenominator);
+			static FFrameRate ConvertToUnrealFramerate(NTV2VideoFormat InVideoFormat);
 			static TimecodeFormat ConvertToTimecodeFormat(NTV2VideoFormat InVideoFormat);
 
 			static bool TryVideoFormatIndexToNTV2VideoFormat(FAJAVideoFormat InVideoFormatIndex, NTV2VideoFormat& OutFoundVideoFormat);
diff --git a/Engine/Plugins/Media/AjaMedia/Source/Aja/Private/OutputChannel.cpp b/Engine/Plugins/Media/AjaMedia/Source/Aja/Private/OutputChannel.cpp
--- a/Engine/Plugins/Media/AjaMedia/Source/Aja/Private/OutputChannel.cpp
+++ b/Engine/Plugins/Media/AjaMedia/Source/Aja/Private/OutputChannel.cpp
@@ -11,6 +11,21 @@ DECLARE_FLOAT_COUNTER_STAT(TEXT("AJA Audio Delay (s)"), STAT_AjaMediaCapture_Aud
 
 #include "Misc/ScopeLock.h"
 
+static TAutoConsoleVariable<int32> CVarAjaPreRollMethod(
+	TEXT("Aja.PrerollMethod"), 0,
+	TEXT("Preroll method 0: Legacy method where we preroll 3 empty frames. Preroll method 1: We preroll engine frames as they come in."),
+	ECVF_RenderThreadSafe);
+
+static TAutoConsoleVariable<int32> CVarAjaPingPongVersion(
+	TEXT("Aja.PingPongVersion"), 0,
+	TEXT("Version 0: Legacy, Version 1: Experimental with reduced waiting."),
+	ECVF_RenderThreadSafe);
+
+static TAutoConsoleVariable<bool> CVarAjaShowDebugMarkers(
+	TEXT("Aja.EnableDebugMarkers"), false,
+	TEXT("Enable to let Aja Output Channel trace more information when doing a media capture, such as bookmarks for Vertical Interrupts etc."),
+	ECVF_RenderThreadSafe);
+
 namespace AJA
 {
 	namespace Private
@@ -229,6 +244,15 @@ namespace AJA
 					FrameReadyToWrite.push_back(NewFrame);
 				}
 			}
+
+			FrameAvailableEvent = FPlatformProcess::GetSynchEventFromPool();
+			bOutputWithReduceWaiting = CVarAjaPingPongVersion.GetValueOnAnyThread() == 1;
+			bTraceDebugMarkers = CVarAjaShowDebugMarkers.GetValueOnAnyThread();
+		}
+
+		OutputChannelThread::~OutputChannelThread()
+		{
+			FPlatformProcess::ReturnSynchEventToPool(FrameAvailableEvent);
 		}
 
 		bool OutputChannelThread::CanInitialize() const
@@ -240,6 +264,11 @@ namespace AJA
 		{
 			ChannelThreadBase::Uninitialize();
 
+			if (InterruptEventHandle.IsValid())
+			{
+				Device->ClearOnInterruptEvent(Channel, InterruptEventHandle);
+			}
+
 			if (TextureTransfer && GetOptions().bUseGPUDMA)
 			{
 				std::vector<uint8_t*> BuffersToClear;
@@ -516,6 +545,7 @@ namespace AJA
 				{
 					FrameReadyToWrite.erase(FoundItt);
 					FrameReadyToRead.push_back(InFrame);
+					FrameAvailableEvent->Trigger();
 				}
 			}
 		}
@@ -633,6 +663,10 @@ namespace AJA
 			Frame* AvailableWritingFrame = FetchAvailableWritingFrame(InFrameData);
 			if (!AvailableWritingFrame)
 			{
+				if (bTraceDebugMarkers)
+				{
+					TRACE_BOOKMARK(TEXT("Aja: Dropped frame %d"), InFrameData.FrameIdentifier);
+				}
 				return false;
 			}
 
@@ -681,12 +715,17 @@ namespace AJA
 #endif
 
 			AvailableWritingFrame->CopiedVideoBufferSize = InVideoBufferSize;
+
 			PushWhenFrameReady(AvailableWritingFrame);
 
+			if (bTraceDebugMarkers)
+			{
+				TRACE_BOOKMARK(TEXT("Aja: Pushed [%d] (Frame %d)"), InFrameData.FrameIdentifier, InFrameData.Timecode.Frames);
+			}
+
 			return true;
 		}
 
-
 		bool OutputChannelThread::SetVideoFrameData(const AJAOutputFrameBufferData& InFrameData, FRHITexture* RHITexture)
 		{
 			const bool bIsProgressive = ::IsProgressivePicture(VideoFormat) || Options.bOutputInterlaceAsProgressive;
@@ -853,9 +892,9 @@ namespace AJA
 			bool bResult = Super::DeviceThread_ConfigureAudio(InCommandList);
 			if(bResult && UseAudio())
 			{
-				NTV2FrameRate FrameRate = NTV2_FRAMERATE_INVALID;
-				AJA_CHECK(Device->GetCard()->GetFrameRate(FrameRate, Channel));
-				NumSamplesPerFrame = ::GetAudioSamplesPerFrame(FrameRate, NTV2_AUDIO_48K);
+				NTV2FrameRate CardFrameRate = NTV2_FRAMERATE_INVALID;
+				AJA_CHECK(Device->GetCard()->GetFrameRate(CardFrameRate, Channel));
+				NumSamplesPerFrame = ::GetAudioSamplesPerFrame(CardFrameRate, NTV2_AUDIO_48K);
 
 				if (ensureMsgf(AudioBufferSize > 0, TEXT("Audio could not be initialized.")))
 				{
@@ -893,6 +932,8 @@ namespace AJA
 		{
 			bool bResult = Super::DeviceThread_ConfigureVideo(InCommandList);
 
+			FrameRate = Helpers::ConvertToUnrealFramerate(DesiredVideoFormat);
+
 			if (VideoBufferSize > 0)
 			{
 				AJAAutoLock AutoLock(&FrameLock);
@@ -907,6 +948,7 @@ namespace AJA
 #endif
 				}
 			}
+
 			return bResult;
 		}
 
@@ -1142,49 +1184,63 @@ namespace AJA
 			}
 
 			ThreadLock_PingPongOutputLoop_Memset0(CurrentFrame);
-
+		
 			AJA_CHECK(Device->WaitForInputOrOutputInterrupt(Channel));
 
-			// Clear the current Buffer and initialize the ping-pong
-			const int32_t NumOutputFrameIndex = 2;
-			for (int32_t Index = 0; Index < NumOutputFrameIndex; ++Index)
+			if (bTraceDebugMarkers)
 			{
-				AJA_CHECK(GetDevice().SetOutputFrame(Channel, BaseFrameIndex + Index));
+				InterruptEventHandle = Device->SetOnInterruptEvent(Channel, DeviceConnection::FOnVerticalInterrupt::FDelegate::CreateRaw(this, &OutputChannelThread::OnVerticalInterrupt));
+			}
 
-				//@TODO fix ancillary for pingpong. AJA told us they would help
-				if (UseAncillary())
-				{
-					NTV2_POINTER AncBufferPointer(CurrentFrame->AncBuffer, CurrentFrame->CopiedAncBufferSize);
-					GetDevice().DMAWriteAnc(Index, AncBufferPointer);
-				}
-				if (UseAncillaryField2())
-				{
-					NTV2_POINTER AncF2BufferPointer(CurrentFrame->AncF2Buffer, CurrentFrame->CopiedAncF2BufferSize);
-					GetDevice().DMAWriteAnc(Index, AncF2BufferPointer);
-				}
-				if (UseAudio())
-				{
-					GetDevice().DMAWriteAudio(AudioSystem, reinterpret_cast<ULWord*>(CurrentFrame->AudioBuffer), 0, CurrentFrame->CopiedAudioBufferSize);
-					CurrentAudioWriteOffset += CurrentFrame->CopiedAudioBufferSize;
-				}
-				if (UseVideo())
+			if (CVarAjaPreRollMethod.GetValueOnAnyThread() == 0)
+			{
+				// Clear the current Buffer and initialize the ping-pong
+				const int32_t NumOutputFrameIndex = 2;
+				for (int32_t Index = 0; Index < NumOutputFrameIndex; ++Index)
 				{
-					AJA_CHECK(GetDevice().DMAWriteFrame(BaseFrameIndex + Index, reinterpret_cast<ULWord*>(CurrentFrame->VideoBuffer), CurrentFrame->CopiedVideoBufferSize, Channel));
+					AJA_CHECK(GetDevice().SetOutputFrame(Channel, BaseFrameIndex + Index));
 
-					if (TextureTransfer && GetOptions().bUseGPUDMA)
+					//@TODO fix ancillary for pingpong. AJA told us they would help
+					if (UseAncillary())
 					{
-						TextureTransfer->EndSync(CurrentFrame->VideoBuffer);
+						NTV2_POINTER AncBufferPointer(CurrentFrame->AncBuffer, CurrentFrame->CopiedAncBufferSize);
+						GetDevice().DMAWriteAnc(Index, AncBufferPointer);
 					}
+					if (UseAncillaryField2())
+					{
+						NTV2_POINTER AncF2BufferPointer(CurrentFrame->AncF2Buffer, CurrentFrame->CopiedAncF2BufferSize);
+						GetDevice().DMAWriteAnc(Index, AncF2BufferPointer);
+					}
+					if (UseAudio())
+					{
+						GetDevice().DMAWriteAudio(AudioSystem, reinterpret_cast<ULWord*>(CurrentFrame->AudioBuffer), 0, CurrentFrame->CopiedAudioBufferSize);
+						CurrentAudioWriteOffset += CurrentFrame->CopiedAudioBufferSize;
+					}
+					if (UseVideo())
+					{
+						AJA_CHECK(GetDevice().DMAWriteFrame(BaseFrameIndex + Index, reinterpret_cast<ULWord*>(CurrentFrame->VideoBuffer), CurrentFrame->CopiedVideoBufferSize, Channel));
+
+						if (TextureTransfer && GetOptions().bUseGPUDMA)
+						{
+							TextureTransfer->EndSync(CurrentFrame->VideoBuffer);
+						}
+					}
+					AJA_CHECK(Device->WaitForInputOrOutputInterrupt(Channel));
 				}
-				AJA_CHECK(Device->WaitForInputOrOutputInterrupt(Channel));
-			}
 
-			//	Before the main loop starts, ping-pong the buffers so the hardware will use
-			//	different buffers than the ones it was using while idling...
-			uint32_t CurrentOutFrame = 0;
-			CurrentOutFrame ^= 1;
-			AJA_CHECK(GetDevice().SetOutputFrame(Channel, BaseFrameIndex + CurrentOutFrame));
+				ULWord VBICount = 0;
+				AJA_CHECK(GetDevice().GetOutputVerticalInterruptCount(VBICount, Channel));
+
+				// SetOutputFrame will tell the card that a frame should be transferred starting on the next VBI, so it will be transferring starting at CurrentVBI + 1, and be done transferring at CurrentVBI + 2
+				BufferFreeIndex[0] = VBICount + 2;
+				BufferFreeIndex[1] = BufferFreeIndex[0] + 1;
 
+				//	Before the main loop starts, ping-pong the buffers so the hardware will use
+				//	different buffers than the ones it was using while idling...
+				uint32_t CurrentOutFrame = 0;
+				CurrentOutFrame ^= 1;
+				AJA_CHECK(GetDevice().SetOutputFrame(Channel, BaseFrameIndex + CurrentOutFrame));
+			}
 
 			return true;
 		}
@@ -1220,6 +1276,7 @@ namespace AJA
 			const bool bIsProgressive = ::IsProgressivePicture(VideoFormat) || Options.bOutputInterlaceAsProgressive;
 
 			uint32_t CurrentOutFrame = 0;
+			// If prerolling normally, then we're targetting the free buffer.
 			CurrentOutFrame ^= 1;
 
 			if (TextureTransfer)
@@ -1228,20 +1285,61 @@ namespace AJA
 			}
 
 			uint32_t LastFrameDropCount = 0;
-			auto LogDropFrames = [&LastFrameDropCount, this]() 
+			auto LogDropFrames = [&LastFrameDropCount, this]()
 			{
 				UE_LOG(LogAjaCore, Warning, TEXT("PingPong: Dropped %d frames while outputting on channel %d on device %S.\n"), LastFrameDropCount, uint32_t(Channel) + 1, GetDevice().GetDisplayName().c_str());
 			};
 
-
 			bool bRunning = true;
 			while (bRunning && !bStopRequested)
 			{
+				if (bOutputWithReduceWaiting)
+				{
+					const double MaxWaitTime = FrameRate.AsInterval();
+					{
+						TRACE_CPUPROFILER_EVENT_SCOPE(AjaOutputChannelThread::WaitingForAvailableFrame);
+						if (!FrameAvailableEvent->Wait(MaxWaitTime * 1000))
+						{
+							LastFrameDropCount++;
+							continue;
+						}
+					}
+				}
+
 				//	Wait until the input has completed capturing a frame...
 				bool bIsValidWaitForOutput = true;
 				if (bIsProgressive)
 				{
-					bRunning = Device->WaitForInputOrOutputInterrupt(Channel);
+					if (bOutputWithReduceWaiting)
+					{
+						ULWord VBICount = 0;
+						AJA_CHECK(GetDevice().GetOutputVerticalInterruptCount(VBICount, Channel));
+
+						// If we're still using the card buffer, wait until next VBI, otherwise go ahead.
+						const uint32 OffScreenBuffer = CurrentOutFrame ^ 1;
+
+						if (VBICount + 2 == BufferFreeIndex[CurrentOutFrame])
+						{
+							TRACE_CPUPROFILER_EVENT_SCOPE(AjaOutputChannelThread::DroppedFrame);
+							// We have already scheduled an output in the same VBI period, this usually occurs after a hitch, so we should drop one of our frames 
+							if (Frame* AvailableReadingFrame = Thread_FetchAvailableReadingFrame())
+							{
+								Thread_PushAvailableReadingFrame(AvailableReadingFrame);
+							}
+
+							continue;
+						}
+						else if (VBICount < BufferFreeIndex[OffScreenBuffer])
+						{
+							// Off screen buffer is not available, so wait until a VBI occurs so we know it's free.
+							TRACE_CPUPROFILER_EVENT_SCOPE(AjaOutputChannelThread::WaitForInputOrOutputInterrupt);
+							bRunning = Device->WaitForInputOrOutputInterrupt(Channel);
+						}
+					}
+					else
+					{
+						bRunning = Device->WaitForInputOrOutputInterrupt(Channel);
+					}
 				}
 				else
 				{
@@ -1294,7 +1392,7 @@ namespace AJA
 						LogDropFrames();
 						LastFrameDropCount = 0;
 					}
-					
+
 					Thread_TestInterlacedOutput(AvailableReadingFrame);
 
 					//	Flip sense of the buffers again to refer to the buffers that the hardware isn't using (i.e. the off-screen buffers)...
@@ -1325,6 +1423,10 @@ namespace AJA
 						bRunning = bRunning && GetDevice().DMAWriteAnc(CurrentOutFrame, AncF2BufferPointer);
 					}
 
+					ULWord VBICount = 0;
+					AJA_CHECK(GetDevice().GetOutputVerticalInterruptCount(VBICount, Channel));
+
+
 					if (AvailableReadingFrame->CopiedVideoBufferSize > 0)
 					{
 						if (bIsProgressive)
@@ -1335,13 +1437,17 @@ namespace AJA
 						{
 							BurnTimecode(AvailableReadingFrame->Timecode, AvailableReadingFrame->Timecode2, AvailableReadingFrame->VideoBuffer);
 						}
+
 						if (TextureTransfer && GetOptions().bUseGPUDMA)
 						{
 							const UE::GPUTextureTransfer::ETransferDirection Direction = UE::GPUTextureTransfer::ETransferDirection::GPU_TO_CPU;
 							TextureTransfer->BeginSync(AvailableReadingFrame->VideoBuffer, Direction);
 						}
 
-						bRunning = bRunning && GetDevice().DMAWriteFrame(BaseFrameIndex + CurrentOutFrame, reinterpret_cast<ULWord*>(AvailableReadingFrame->VideoBuffer), AvailableReadingFrame->CopiedVideoBufferSize, Channel);
+						{
+							TRACE_CPUPROFILER_EVENT_SCOPE(AjaOutputChannelThread::DMAWriteFrame);
+							bRunning = bRunning && GetDevice().DMAWriteFrame(BaseFrameIndex + CurrentOutFrame, reinterpret_cast<ULWord*>(AvailableReadingFrame->VideoBuffer), AvailableReadingFrame->CopiedVideoBufferSize, Channel);
+						}
 
 						if (TextureTransfer && GetOptions().bUseGPUDMA)
 						{
@@ -1357,7 +1463,6 @@ namespace AJA
 					if (!bRunning)
 					{
 						UE_LOG(LogAjaCore, Error, TEXT("PingPong: Can't do the DMA frame transfer for channel %d on device %S.\n"), uint32_t(Channel) + 1, GetDevice().GetDisplayName().c_str());
-						Thread_PushAvailableReadingFrame(AvailableReadingFrame);
 						break;
 					}
 
@@ -1382,7 +1487,25 @@ namespace AJA
 					}
 
 					//	Tell the hardware which buffers to start using at the beginning of the next frame...
+					if (bTraceDebugMarkers)
+					{
+						TRACE_BOOKMARK(TEXT("Aja: SetOutputFrame [%d] (Buffer %d)"), AvailableReadingFrame->FrameIdentifier, CurrentOutFrame);
+					}
+
 					AJA_CHECK(GetDevice().SetOutputFrame(Channel, BaseFrameIndex + CurrentOutFrame));
+
+					if (!bRunning)
+					{
+						UE_LOG(LogAjaCore, Error, TEXT("PingPong: Can't do the DMA frame transfer for channel %d on device %S.\n"), uint32_t(Channel) + 1, GetDevice().GetDisplayName().c_str());
+						Thread_PushAvailableReadingFrame(AvailableReadingFrame);
+						break;
+					}
+
+					// SetOutputFrame will tell the card that a frame should be transferred starting on the next VBI, so it will be transferring starting at CurrentVBI + 1, and be done transferring at CurrentVBI + 2
+					uint32 FrameAvailableVBICount = VBICount + 2;
+					BufferFreeIndex[CurrentOutFrame] = FrameAvailableVBICount;
+					BufferFrameIdMapping[CurrentOutFrame] = AvailableReadingFrame->FrameIdentifier;
+
 					Thread_PushAvailableReadingFrame(AvailableReadingFrame);
 					bHaveOutputOnce = true;
 
@@ -1424,9 +1547,9 @@ namespace AJA
 					testPatternBuffer);
 
 				CurrentOutFrame ^= 1;
-				GetDevice().DMAWriteFrame(BaseFrameIndex + CurrentOutFrame, reinterpret_cast <uint32_t *> (testPatternBuffer.GetHostPointer()), uint32_t(testPatternBuffer.GetByteCount()));
+				GetDevice().DMAWriteFrame(BaseFrameIndex + CurrentOutFrame, reinterpret_cast <uint32_t*> (testPatternBuffer.GetHostPointer()), uint32_t(testPatternBuffer.GetByteCount()));
 				CurrentOutFrame ^= 1;
-				GetDevice().DMAWriteFrame(BaseFrameIndex + CurrentOutFrame, reinterpret_cast <uint32_t *> (testPatternBuffer.GetHostPointer()), uint32_t(testPatternBuffer.GetByteCount()));
+				GetDevice().DMAWriteFrame(BaseFrameIndex + CurrentOutFrame, reinterpret_cast <uint32_t*> (testPatternBuffer.GetHostPointer()), uint32_t(testPatternBuffer.GetByteCount()));
 			}
 
 			if (!bStopRequested)
@@ -1725,5 +1848,31 @@ namespace AJA
 			}
 			return false;
 		}
+
+		void OutputChannelThread::OnVerticalInterrupt(uint32 SyncCount) const
+		{
+			uint32 FrameId = 0;
+			uint8 BufferId = 0;
+
+			if (BufferFreeIndex[0] == SyncCount)
+			{
+				FrameId = BufferFrameIdMapping[0];
+				BufferId = 0;
+			}
+			else if (BufferFreeIndex[1] == SyncCount)
+			{
+				FrameId = BufferFrameIdMapping[1];
+				BufferId = 1;
+			}
+
+			if (FrameId != 0)
+			{
+				TRACE_BOOKMARK(TEXT("Aja: Finish SDI transfer frame [%d] (Buffer: %d)"), FrameId, BufferId);
+			}
+			else if (BufferFrameIdMapping[0] != 0 && BufferFrameIdMapping[1] != 0) // Don't log drop frames when starting the output
+			{
+				TRACE_BOOKMARK(TEXT("Aja VBI [%d] (Dropped frame)"), SyncCount);
+			}
+		}
 	}
 }
\ No newline at end of file
diff --git a/Engine/Plugins/Media/AjaMedia/Source/Aja/Private/OutputChannel.h b/Engine/Plugins/Media/AjaMedia/Source/Aja/Private/OutputChannel.h
--- a/Engine/Plugins/Media/AjaMedia/Source/Aja/Private/OutputChannel.h
+++ b/Engine/Plugins/Media/AjaMedia/Source/Aja/Private/OutputChannel.h
@@ -41,7 +41,12 @@ namespace AJA
 		private:
 			std::shared_ptr<OutputChannelThread> ChannelThread;
 			std::shared_ptr<DeviceConnection> Device;
+			/** VBI Buffer count where [x] will be done transferring on the wire. */
+			uint32 BufferFreeIndex[2] = { 0, 0 };
+			/** Debug information used to map a VBI event to a UE frame number. */
+			uint32 BufferFrameIdMapping[2] = { 0, 0 };
 			std::weak_ptr<IOChannelInitialize_DeviceCommand> InitializeCommand;
+			uint32_t CurrentOutFrame = 0;
 		};
 
 		/* OutputChannel OutputChannelThread
@@ -52,6 +57,7 @@ namespace AJA
 			using Super = ChannelThreadBase;
 		public:
 			OutputChannelThread(const AJADeviceOptions& InDevice, const AJAInputOutputChannelOptions& InOptions);
+			~OutputChannelThread();
 			
 			virtual bool CanInitialize() const override;
 			virtual void Uninitialize() override;
@@ -94,6 +100,9 @@ namespace AJA
 			bool Thread_HandleLostFrameAudio();
 			bool Thread_ShouldPauseAudioOutput();
 
+			/** Callback on the vertical interrupt used to keep track of which buffers are currently outputting and to trace debug information. */
+			void OnVerticalInterrupt(uint32 SyncCount) const;
+
 			struct Frame
 			{
 				Frame();
@@ -132,10 +141,16 @@ namespace AJA
 			std::vector<Frame*> FrameReadyToRead; // Ready to be copied on the AJA
 			std::vector<Frame*> FrameReadyToWrite; // Ready to be used by Unreal Engine
 
+
 			bool bIsSSE2Available;
 			bool bIsFieldAEven;
 			bool bIsFirstFetch;
 
+			/** Enabled through Aja.PingPongVersion 1, will wait on a frame received event before falling back on a WaitForVerticalInterrupt call. */
+			bool bOutputWithReduceWaiting = false;
+			/** Indicates if OutputChannel should trace debug information. (Disabled by default to avoid cluttering Insights Traces). */
+			bool bTraceDebugMarkers = false;
+
 			bool bInterlacedTest_ExpectFirstLineToBeWhite;
 			uint32_t InterlacedTest_FrameCounter;
 
@@ -149,6 +164,21 @@ namespace AJA
 			FCriticalSection DeviceCriticalSection;
 
 			UE::GPUTextureTransfer::TextureTransferPtr TextureTransfer = nullptr;
+
+			/** VBI count where BufferFreeIndex[x] will be done transferring on the wire. */
+			uint32 BufferFreeIndex[2] = { 0, 0 };
+
+			/** Debug information used to map a VBI event to a UE frame number. */
+			uint32 BufferFrameIdMapping[2] = { 0, 0 };
+
+			/** Event triggered whenever a frame is pushed by Unreal to the FrameReadyToRead queue. */
+			FEvent* FrameAvailableEvent = nullptr;
+
+			/** Output framerate. */
+			FFrameRate FrameRate;
+
+			/** Handle to the delegate that's triggered when a card emits a vertical interrupt event. */
+			FDelegateHandle InterruptEventHandle;
 		};
 	}
 }
diff --git a/Engine/Plugins/Media/AjaMedia/Source/AjaMediaOutput/Private/AjaMediaCapture.cpp b/Engine/Plugins/Media/AjaMedia/Source/AjaMediaOutput/Private/AjaMediaCapture.cpp
--- a/Engine/Plugins/Media/AjaMedia/Source/AjaMediaOutput/Private/AjaMediaCapture.cpp
+++ b/Engine/Plugins/Media/AjaMedia/Source/AjaMediaOutput/Private/AjaMediaCapture.cpp
@@ -806,7 +806,7 @@ void UAjaMediaCapture::OnRHIResourceCaptured_AnyThread(const FCaptureBaseData& I
 		
 		AJA::AJAOutputFrameBufferData FrameBuffer;
 		FrameBuffer.Timecode = Timecode;
-		FrameBuffer.FrameIdentifier = InBaseData.SourceFrameNumber;
+		FrameBuffer.FrameIdentifier = InBaseData.SourceFrameNumberRenderThread;
 
 		bool bSetVideoResult = false;
 		{
